Merging from CL 11878573.
Add DMX Fixture assets.

diff --git a/Engine/Plugins/Experimental/DigitalMultiplex/DMXFixtures/DMXFixtures.uplugin b/Engine/Plugins/Experimental/DigitalMultiplex/DMXFixtures/DMXFixtures.uplugin
new file mode 100644
--- /dev/null
+++ b/Engine/Plugins/Experimental/DigitalMultiplex/DMXFixtures/DMXFixtures.uplugin
@@ -0,0 +1,17 @@
+{
+  "FileVersion": 3,
+  "Version": 1,
+  "VersionName": "1.0",
+  "FriendlyName": "DMX Fixtures",
+  "Description": "DMX Light Fixtures Blueprints",
+  "Category": "Media",
+  "CreatedBy": "Epic Games, Inc.",
+  "CreatedByURL": "http://epicgames.com",
+  "DocsURL": "",
+  "MarketplaceURL": "",
+  "SupportURL": "",
+  "CanContainContent": true,
+  "IsBetaVersion": true,
+  "IsExperimentalVersion": false,
+  "Installed": false
+}
\ No newline at end of file
