DMX - Fix a bug where DMXFixtureDoubleComponent would not raise SetChannelValue blueprint events

diff --git a/Engine/Plugins/VirtualProduction/DMX/DMXFixtures/Source/DMXFixtures/Private/DMXFixtureComponentDouble.cpp b/Engine/Plugins/VirtualProduction/DMX/DMXFixtures/Source/DMXFixtures/Private/DMXFixtureComponentDouble.cpp
--- a/Engine/Plugins/VirtualProduction/DMX/DMXFixtures/Source/DMXFixtures/Private/DMXFixtureComponentDouble.cpp
+++ b/Engine/Plugins/VirtualProduction/DMX/DMXFixtures/Source/DMXFixtures/Private/DMXFixtureComponentDouble.cpp
@@ -61,7 +61,7 @@ float UDMXFixtureComponentDouble::GetDMXInterpolatedValue(int32 ChannelIndex) co
 		const FInterpolationData& InterpolationData = ChannelIndex == 0 ? CurrentCell->ChannelInterpolation[0] : CurrentCell->ChannelInterpolation[1];
 		return InterpolationData.CurrentValue;
 	}
-	
+
 	return 0.f;
 }
 
@@ -111,34 +111,53 @@ bool UDMXFixtureComponentDouble::IsTargetValid(int32 ChannelIndex, float Target)
 
 		return InterpolationData.IsTargetValid(Target, SkipThreshold);
 	}
-	
+
 	return false;
 }
 
 void UDMXFixtureComponentDouble::SetTargetValue(int32 ChannelIndex, float AbsoluteValue)
 {
-	if (CurrentCell && CurrentCell->ChannelInterpolation.Num() == 2 &&
+	if (CurrentCell && 
+		CurrentCell->ChannelInterpolation.Num() == 2 &&
 		IsTargetValid(ChannelIndex, AbsoluteValue))
 	{
-	FInterpolationData& InterpolationData = ChannelIndex == 0 ? CurrentCell->ChannelInterpolation[0] : CurrentCell->ChannelInterpolation[1];
+		FInterpolationData& InterpolationData = ChannelIndex == 0 ? CurrentCell->ChannelInterpolation[0] : CurrentCell->ChannelInterpolation[1];
 
 		if (bUseInterpolation)
-	{
-			if (InterpolationData.bFirstValueWasSet)
 		{
+			if (InterpolationData.bFirstValueWasSet)
+			{				
 				// Only 'Push' the next value into interpolation. BPs will read the resulting value on tick.
 				InterpolationData.Push(AbsoluteValue);
-		}
-		else
-		{
+			}
+			else
+			{
 				// Jump to the first value if it never was set
 				InterpolationData.SetValueNoInterp(AbsoluteValue);
 				InterpolationData.bFirstValueWasSet = true;
+
+				if (ChannelIndex == 0)
+				{
+					SetChannel1ValueNoInterp(AbsoluteValue);
+				}
+				else
+				{
+					SetChannel2ValueNoInterp(AbsoluteValue);
+				}
+			}
 		}
-	}
-	else
-	{
+		else
+		{
 			InterpolationData.SetValueNoInterp(AbsoluteValue);
+
+			if (ChannelIndex == 0)
+			{
+				SetChannel1ValueNoInterp(AbsoluteValue);
+			}
+			else
+			{
+				SetChannel2ValueNoInterp(AbsoluteValue);
+			}
 		}
 	}
 }	
