DMX - Fix CIS NoPCH Compile LyraEditor Win64

diff --git a/Engine/Plugins/VirtualProduction/DMX/DMXFixtures/Source/DMXFixtures/Private/DMXFixtureActorMatrix.cpp b/Engine/Plugins/VirtualProduction/DMX/DMXFixtures/Source/DMXFixtures/Private/DMXFixtureActorMatrix.cpp
--- a/Engine/Plugins/VirtualProduction/DMX/DMXFixtures/Source/DMXFixtures/Private/DMXFixtureActorMatrix.cpp
+++ b/Engine/Plugins/VirtualProduction/DMX/DMXFixtures/Source/DMXFixtures/Private/DMXFixtureActorMatrix.cpp
@@ -5,6 +5,7 @@
 #include "DMXProtocolConstants.h"
 #include "DMXStats.h"
 
+#include "Engine/Texture2D.h"
 #include "Rendering/Texture2DResource.h"
 #include "Components/StaticMeshComponent.h"
 
