DMX - Fix CIS NoPCH Compile LyraEditor Win64

diff --git a/Engine/Plugins/VirtualProduction/DMX/DMXFixtures/Source/DMXFixtures/Private/DMXFixtureComponent.cpp b/Engine/Plugins/VirtualProduction/DMX/DMXFixtures/Source/DMXFixtures/Private/DMXFixtureComponent.cpp
--- a/Engine/Plugins/VirtualProduction/DMX/DMXFixtures/Source/DMXFixtures/Private/DMXFixtureComponent.cpp
+++ b/Engine/Plugins/VirtualProduction/DMX/DMXFixtures/Source/DMXFixtures/Private/DMXFixtureComponent.cpp
@@ -8,7 +8,9 @@
 
 #include "RenderingThread.h"
 #include "RHICommandList.h"
+#include "Engine/Texture2D.h"
 #include "Kismet/KismetMathLibrary.h"
+#include "Rendering/Texture2DResource.h"
 
 
 UDMXFixtureComponent::UDMXFixtureComponent()
