ObjectPtr upgrade for engine plugins used from UnrealEditor -allmodules.  Automated upgrade process.

diff --git a/Engine/Plugins/Runtime/NetworkPredictionExtras/Source/NetworkPredictionExtras/Public/NetworkPredictionExtrasCharacter.h b/Engine/Plugins/Runtime/NetworkPredictionExtras/Source/NetworkPredictionExtras/Public/NetworkPredictionExtrasCharacter.h
--- a/Engine/Plugins/Runtime/NetworkPredictionExtras/Source/NetworkPredictionExtras/Public/NetworkPredictionExtrasCharacter.h
+++ b/Engine/Plugins/Runtime/NetworkPredictionExtras/Source/NetworkPredictionExtras/Public/NetworkPredictionExtrasCharacter.h
@@ -82,7 +82,7 @@ protected:
 	void ProduceInput(const int32 DeltaMS, FCharacterMotionInputCmd& Cmd);
 
 	UPROPERTY(Category=Movement, VisibleAnywhere)
-	UCharacterMotionComponent* CharacterMotionComponent;
+	TObjectPtr<UCharacterMotionComponent> CharacterMotionComponent;
 
 private:
 
diff --git a/Engine/Plugins/Runtime/NetworkPredictionExtras/Source/NetworkPredictionExtras/Public/NetworkPredictionExtrasFlyingPawn.h b/Engine/Plugins/Runtime/NetworkPredictionExtras/Source/NetworkPredictionExtras/Public/NetworkPredictionExtrasFlyingPawn.h
--- a/Engine/Plugins/Runtime/NetworkPredictionExtras/Source/NetworkPredictionExtras/Public/NetworkPredictionExtrasFlyingPawn.h
+++ b/Engine/Plugins/Runtime/NetworkPredictionExtras/Source/NetworkPredictionExtras/Public/NetworkPredictionExtrasFlyingPawn.h
@@ -82,7 +82,7 @@ protected:
 	void ProduceInput(const int32 DeltaMS, FFlyingMovementInputCmd& Cmd);
 
 	UPROPERTY(Category=Movement, VisibleAnywhere)
-	UFlyingMovementComponent* FlyingMovementComponent;
+	TObjectPtr<UFlyingMovementComponent> FlyingMovementComponent;
 
 private:
 
diff --git a/Engine/Plugins/Runtime/SoundCueTemplates/Source/SoundCueTemplates/Public/SoundCueContainer.h b/Engine/Plugins/Runtime/SoundCueTemplates/Source/SoundCueTemplates/Public/SoundCueContainer.h
--- a/Engine/Plugins/Runtime/SoundCueTemplates/Source/SoundCueTemplates/Public/SoundCueContainer.h
+++ b/Engine/Plugins/Runtime/SoundCueTemplates/Source/SoundCueTemplates/Public/SoundCueContainer.h
@@ -57,7 +57,7 @@ public:
 	bool bLooping;
 
 	UPROPERTY(EditAnywhere, Category = "Variation")
-	TSet<USoundWave*> Variations;
+	TSet<TObjectPtr<USoundWave>> Variations;
 
 	UPROPERTY(EditAnywhere, Category = "Modulation")
 	FVector2D PitchModulation;
diff --git a/Engine/Plugins/Runtime/SoundCueTemplates/Source/SoundCueTemplates/Public/SoundCueDistanceCrossfade.h b/Engine/Plugins/Runtime/SoundCueTemplates/Source/SoundCueTemplates/Public/SoundCueDistanceCrossfade.h
--- a/Engine/Plugins/Runtime/SoundCueTemplates/Source/SoundCueTemplates/Public/SoundCueDistanceCrossfade.h
+++ b/Engine/Plugins/Runtime/SoundCueTemplates/Source/SoundCueTemplates/Public/SoundCueDistanceCrossfade.h
@@ -33,7 +33,7 @@ struct FSoundCueCrossfadeInfo
 	FDistanceDatum DistanceInfo;
 
 	UPROPERTY(EditAnywhere, Category = "Crossfade Info")
-	USoundWave* Sound = nullptr;
+	TObjectPtr<USoundWave> Sound = nullptr;
 };
 
 UCLASS(hidecategories = object, BlueprintType)
diff --git a/Engine/Plugins/Runtime/SoundMod/Source/SoundMod/Classes/SoundModWave.h b/Engine/Plugins/Runtime/SoundMod/Source/SoundMod/Classes/SoundModWave.h
--- a/Engine/Plugins/Runtime/SoundMod/Source/SoundMod/Classes/SoundModWave.h
+++ b/Engine/Plugins/Runtime/SoundMod/Source/SoundMod/Classes/SoundModWave.h
@@ -18,7 +18,7 @@ class USoundModWave : public USoundWaveProcedural
 	GENERATED_UCLASS_BODY()
 
 	UPROPERTY()
-	class USoundMod* SoundMod;
+	TObjectPtr<class USoundMod> SoundMod;
 
 	xmp_context xmpContext;
 
diff --git a/Engine/Plugins/Runtime/SoundMod/Source/SoundMod/Public/SoundNodeModPlayer.h b/Engine/Plugins/Runtime/SoundMod/Source/SoundMod/Public/SoundNodeModPlayer.h
--- a/Engine/Plugins/Runtime/SoundMod/Source/SoundMod/Public/SoundNodeModPlayer.h
+++ b/Engine/Plugins/Runtime/SoundMod/Source/SoundMod/Public/SoundNodeModPlayer.h
@@ -29,7 +29,7 @@ private:
 	TSoftObjectPtr<USoundMod> SoundModAssetPtr;
 
 	UPROPERTY(transient)
-	USoundMod* SoundMod;
+	TObjectPtr<USoundMod> SoundMod;
 
 	UPROPERTY(EditAnywhere, Category=ModPlayer)
 	uint32 bLooping:1;
diff --git a/Engine/Plugins/Runtime/SoundUtilities/Source/SoundUtilities/Public/SoundSimple.h b/Engine/Plugins/Runtime/SoundUtilities/Source/SoundUtilities/Public/SoundSimple.h
--- a/Engine/Plugins/Runtime/SoundUtilities/Source/SoundUtilities/Public/SoundSimple.h
+++ b/Engine/Plugins/Runtime/SoundUtilities/Source/SoundUtilities/Public/SoundSimple.h
@@ -15,7 +15,7 @@ struct SOUNDUTILITIES_API FSoundVariation
 
 	// The sound wave asset to use for this variation
 	UPROPERTY(EditAnywhere, BlueprintReadWrite, Category = "SoundVariation")
-	USoundWave* SoundWave;
+	TObjectPtr<USoundWave> SoundWave;
 
 	// The probability weight to use for this variation
 	UPROPERTY(EditAnywhere, BlueprintReadWrite, Category = "Synth|Preset")
@@ -68,7 +68,7 @@ protected:
 
 	// The current chosen sound wave
 	UPROPERTY(transient)
-	USoundWave* SoundWave;
+	TObjectPtr<USoundWave> SoundWave;
 };
 
 
diff --git a/Engine/Plugins/Runtime/SoundUtilities/Source/SoundUtilitiesEditor/Private/SoundSimpleFactory.h b/Engine/Plugins/Runtime/SoundUtilities/Source/SoundUtilitiesEditor/Private/SoundSimpleFactory.h
--- a/Engine/Plugins/Runtime/SoundUtilities/Source/SoundUtilitiesEditor/Private/SoundSimpleFactory.h
+++ b/Engine/Plugins/Runtime/SoundUtilities/Source/SoundUtilitiesEditor/Private/SoundSimpleFactory.h
@@ -20,6 +20,6 @@ class USoundSimpleFactory : public UFactory
 
 	/** Sound waves to create the simple sound with */
 	UPROPERTY()
-	TArray<USoundWave*> SoundWaves;
+	TArray<TObjectPtr<USoundWave>> SoundWaves;
 };
 
diff --git a/Engine/Plugins/Runtime/Soundscape/Source/SoundScape/Public/SoundScapePalette.h b/Engine/Plugins/Runtime/Soundscape/Source/SoundScape/Public/SoundScapePalette.h
--- a/Engine/Plugins/Runtime/Soundscape/Source/SoundScape/Public/SoundScapePalette.h
+++ b/Engine/Plugins/Runtime/Soundscape/Source/SoundScape/Public/SoundScapePalette.h
@@ -19,7 +19,7 @@ struct SOUNDSCAPE_API FSoundscapePaletteColor
 
 	// Soundscape Color to Play
 	UPROPERTY(EditAnywhere, BlueprintReadWrite, Category = "Soundscape|Palette")
-	USoundscapeColor* SoundscapeColor = nullptr;
+	TObjectPtr<USoundscapeColor> SoundscapeColor = nullptr;
 
 	// Base Volume Scalar
 	UPROPERTY(EditAnywhere, BlueprintReadWrite, Category = "Soundscape|Palette", meta = (ClampMin = "0.0", ClampMax = "4.0", UIMin = "0.0", UIMax = "4.0", SliderExponent = "6.0"))
@@ -77,9 +77,9 @@ public:
 
 private:
 	UPROPERTY()
-	UWorld* World;
+	TObjectPtr<UWorld> World;
 
 	UPROPERTY()
-	TArray<UActiveSoundscapeColor*> ActiveSoundscapeColors;
+	TArray<TObjectPtr<UActiveSoundscapeColor>> ActiveSoundscapeColors;
 
 };
diff --git a/Engine/Plugins/Runtime/Soundscape/Source/SoundScape/Public/SoundscapeColor.h b/Engine/Plugins/Runtime/Soundscape/Source/SoundScape/Public/SoundscapeColor.h
--- a/Engine/Plugins/Runtime/Soundscape/Source/SoundScape/Public/SoundscapeColor.h
+++ b/Engine/Plugins/Runtime/Soundscape/Source/SoundScape/Public/SoundscapeColor.h
@@ -266,7 +266,7 @@ public:
 
 	// Sound to Play
 	UPROPERTY(EditAnywhere, BlueprintReadWrite, Category = "Soundscape|Color")
-	USoundBase* Sound;
+	TObjectPtr<USoundBase> Sound;
 
 	// Base Volume Scalar
 	UPROPERTY(EditAnywhere, BlueprintReadWrite, Category = "Soundscape|Color", meta = (ClampMin = "0.0", ClampMax = "4.0", UIMin = "0.0", UIMax = "4.0", SliderExponent = "6.0"))
@@ -306,7 +306,7 @@ public:
 	//~UObject Interface End
 
 	UPROPERTY(EditAnywhere, Category = "Soundscape|Color")
-	UAudioComponent* AudioComponent = nullptr;
+	TObjectPtr<UAudioComponent> AudioComponent = nullptr;
 
 	UPROPERTY(EditAnywhere, Category = "Soundscape|Color")
 	FTimerHandle TimerHandle;
@@ -354,7 +354,7 @@ public:
 private:
 	// Sound to Play
 	UPROPERTY()
-	USoundBase* Sound;
+	TObjectPtr<USoundBase> Sound;
 
 	// Base Volume Scalar
 	float VolumeBase = 1.0f;
@@ -443,7 +443,7 @@ private:
 
 	// Audio Component Pool
 	UPROPERTY()
-	TArray<UActiveSoundscapeColorVoice*> SoundscapeColorVoicePool;
+	TArray<TObjectPtr<UActiveSoundscapeColorVoice>> SoundscapeColorVoicePool;
 
 	// Is Playing Bool
 	bool bIsPlaying = false;
diff --git a/Engine/Plugins/Runtime/Soundscape/Source/SoundScape/Public/SoundscapeSubsystem.h b/Engine/Plugins/Runtime/Soundscape/Source/SoundScape/Public/SoundscapeSubsystem.h
--- a/Engine/Plugins/Runtime/Soundscape/Source/SoundScape/Public/SoundscapeSubsystem.h
+++ b/Engine/Plugins/Runtime/Soundscape/Source/SoundScape/Public/SoundscapeSubsystem.h
@@ -33,7 +33,7 @@ struct SOUNDSCAPE_API FSoundscapePaletteCollectionLoaded
 
 	// Soundscape Palette Collection
 	UPROPERTY(EditAnywhere, Category = "Soundscape|Palette")
-	TSet<USoundscapePalette*> SoundscapePaletteCollection;
+	TSet<TObjectPtr<USoundscapePalette>> SoundscapePaletteCollection;
 };
 
 // Struct
@@ -104,7 +104,7 @@ protected:
 private:
 	// Color Point Hash Map
 	UPROPERTY()
-	TMap<FGameplayTag, USoundscapeColorPointHashCellDensity*> ColorPointHashMap;
+	TMap<FGameplayTag, TObjectPtr<USoundscapeColorPointHashCellDensity>> ColorPointHashMap;
 
 	// Cubic Root of the Max Value of int64 rounded down to the nearest million
 	static const int32 MaxGridWidth = 2000000;
@@ -169,11 +169,11 @@ private:
 
 	// Color Point Hash Maps
 	UPROPERTY()
-	USoundscapeColorPointHashMap* ColorPointHashMapLOD1;
+	TObjectPtr<USoundscapeColorPointHashMap> ColorPointHashMapLOD1;
 	UPROPERTY()
-	USoundscapeColorPointHashMap* ColorPointHashMapLOD2;
+	TObjectPtr<USoundscapeColorPointHashMap> ColorPointHashMapLOD2;
 	UPROPERTY()
-	USoundscapeColorPointHashMap* ColorPointHashMapLOD3;
+	TObjectPtr<USoundscapeColorPointHashMap> ColorPointHashMapLOD3;
 
 	// Hash Cell Width for LOD1
 	UPROPERTY()
@@ -235,7 +235,7 @@ public:
 
 private:
 	UPROPERTY()
-	TSet<USoundscapePalette*> LoadedPaletteCollectionSet;
+	TSet<TObjectPtr<USoundscapePalette>> LoadedPaletteCollectionSet;
 
 	UPROPERTY()
 	TMap<FName, FSoundscapePaletteCollection> UnloadedPaletteCollections;
@@ -243,7 +243,7 @@ private:
 	FGameplayTagContainer SubsystemState;
 
 	UPROPERTY()
-	TMap<USoundscapePalette*, UActiveSoundscapePalette*> ActivePalettes;
+	TMap<TObjectPtr<USoundscapePalette>, TObjectPtr<UActiveSoundscapePalette>> ActivePalettes;
 
 	bool LoadPaletteCollection(FName PaletteCollectionName);
 
@@ -290,7 +290,7 @@ public:
 private:
 	// Stored Hash Map Collections
 	UPROPERTY(Transient)
-	TArray<USoundscapeColorPointHashMapCollection*> ColorPointHashMapCollections;
+	TArray<TObjectPtr<USoundscapeColorPointHashMapCollection>> ColorPointHashMapCollections;
 
 	// Stored Color Point Collections
 	TMap<FName, FSoundscapeColorPointCollection> ColorPointCollections;
@@ -300,14 +300,14 @@ private:
 
 	// Color Point Hash Maps
 	UPROPERTY()
-	USoundscapeColorPointHashMap* ColorPointHashMapLOD1;
+	TObjectPtr<USoundscapeColorPointHashMap> ColorPointHashMapLOD1;
 	UPROPERTY()
-	USoundscapeColorPointHashMap* ColorPointHashMapLOD2;
+	TObjectPtr<USoundscapeColorPointHashMap> ColorPointHashMapLOD2;
 	UPROPERTY()
-	USoundscapeColorPointHashMap* ColorPointHashMapLOD3;
+	TObjectPtr<USoundscapeColorPointHashMap> ColorPointHashMapLOD3;
 
 	UPROPERTY()
-	USoundscapeColorPointHashMap* ActiveColorPointHashMap;
+	TObjectPtr<USoundscapeColorPointHashMap> ActiveColorPointHashMap;
 
 	// Hash Cell Width for LOD1
 	float LOD1ColorPointHashWidth = 500.0f;
diff --git a/Engine/Plugins/Runtime/StateTree/Source/StateTreeEditorModule/Public/StateTreeCompilerLog.h b/Engine/Plugins/Runtime/StateTree/Source/StateTreeEditorModule/Public/StateTreeCompilerLog.h
--- a/Engine/Plugins/Runtime/StateTree/Source/StateTreeEditorModule/Public/StateTreeCompilerLog.h
+++ b/Engine/Plugins/Runtime/StateTree/Source/StateTreeEditorModule/Public/StateTreeCompilerLog.h
@@ -30,7 +30,7 @@ struct STATETREEEDITORMODULE_API FStateTreeCompilerLogMessage
 
 	/** (optional) The StateTree state the message refers to. */
 	UPROPERTY()
-	const UStateTreeState* State = nullptr;
+	TObjectPtr<const UStateTreeState> State = nullptr;
 
 	/** (optional) The State tee item (condition/evaluator/task) the message refers to. */
 	UPROPERTY()
@@ -97,7 +97,7 @@ struct STATETREEEDITORMODULE_API FStateTreeCompilerLog
 	
 protected:
 	UPROPERTY()
-	TArray<const UStateTreeState*> StateStack;
+	TArray<TObjectPtr<const UStateTreeState>> StateStack;
 	
 	UPROPERTY()
 	TArray<FStateTreeCompilerLogMessage> Messages;
diff --git a/Engine/Plugins/Runtime/StateTree/Source/StateTreeEditorModule/Public/StateTreeEditorData.h b/Engine/Plugins/Runtime/StateTree/Source/StateTreeEditorModule/Public/StateTreeEditorData.h
--- a/Engine/Plugins/Runtime/StateTree/Source/StateTreeEditorModule/Public/StateTreeEditorData.h
+++ b/Engine/Plugins/Runtime/StateTree/Source/StateTreeEditorModule/Public/StateTreeEditorData.h
@@ -138,5 +138,5 @@ public:
 
 	/** Top level States. */
 	UPROPERTY()
-	TArray<UStateTreeState*> SubTrees;
+	TArray<TObjectPtr<UStateTreeState>> SubTrees;
 };
diff --git a/Engine/Plugins/Runtime/StateTree/Source/StateTreeEditorModule/Public/StateTreeEditorNode.h b/Engine/Plugins/Runtime/StateTree/Source/StateTreeEditorModule/Public/StateTreeEditorNode.h
--- a/Engine/Plugins/Runtime/StateTree/Source/StateTreeEditorModule/Public/StateTreeEditorNode.h
+++ b/Engine/Plugins/Runtime/StateTree/Source/StateTreeEditorModule/Public/StateTreeEditorNode.h
@@ -50,7 +50,7 @@ struct STATETREEEDITORMODULE_API FStateTreeEditorNode
 	FInstancedStruct Instance;
 
 	UPROPERTY(EditDefaultsOnly, Instanced, Category = Node)
-	UObject* InstanceObject = nullptr;
+	TObjectPtr<UObject> InstanceObject = nullptr;
 	
 	UPROPERTY(EditDefaultsOnly, Category = Node)
 	FGuid ID;
diff --git a/Engine/Plugins/Runtime/StateTree/Source/StateTreeEditorModule/Public/StateTreeState.h b/Engine/Plugins/Runtime/StateTree/Source/StateTreeEditorModule/Public/StateTreeState.h
--- a/Engine/Plugins/Runtime/StateTree/Source/StateTreeEditorModule/Public/StateTreeState.h
+++ b/Engine/Plugins/Runtime/StateTree/Source/StateTreeEditorModule/Public/StateTreeState.h
@@ -211,11 +211,11 @@ public:
 	TArray<FStateTreeTransition> Transitions;
 
 	UPROPERTY()
-	TArray<UStateTreeState*> Children;
+	TArray<TObjectPtr<UStateTreeState>> Children;
 
 	UPROPERTY(meta = (ExcludeFromHash))
 	bool bExpanded = true;
 
 	UPROPERTY(meta = (ExcludeFromHash))
-	UStateTreeState* Parent = nullptr;
+	TObjectPtr<UStateTreeState> Parent = nullptr;
 };
diff --git a/Engine/Plugins/Runtime/StateTree/Source/StateTreeModule/Public/StateTreeAnyEnum.h b/Engine/Plugins/Runtime/StateTree/Source/StateTreeModule/Public/StateTreeAnyEnum.h
--- a/Engine/Plugins/Runtime/StateTree/Source/StateTreeModule/Public/StateTreeAnyEnum.h
+++ b/Engine/Plugins/Runtime/StateTree/Source/StateTreeModule/Public/StateTreeAnyEnum.h
@@ -36,6 +36,6 @@ struct STATETREEMODULE_API FStateTreeAnyEnum
 
 	/** The enum class associated with this enum. */
 	UPROPERTY(EditAnywhere, Category = Enum)
-	UEnum* Enum = nullptr;
+	TObjectPtr<UEnum> Enum = nullptr;
 };
 
diff --git a/Engine/Plugins/Runtime/StateTree/Source/StateTreeModule/Public/StateTreeTypes.h b/Engine/Plugins/Runtime/StateTree/Source/StateTreeModule/Public/StateTreeTypes.h
--- a/Engine/Plugins/Runtime/StateTree/Source/StateTreeModule/Public/StateTreeTypes.h
+++ b/Engine/Plugins/Runtime/StateTree/Source/StateTreeModule/Public/StateTreeTypes.h
@@ -611,7 +611,7 @@ struct STATETREEMODULE_API FStateTreeExternalDataDesc
 	
 	/** Class or struct of the external data. */
 	UPROPERTY();
-	const UStruct* Struct = nullptr;
+	TObjectPtr<const UStruct> Struct = nullptr;
 
 	/**
 	 * Name of the external data. Used only for bindable external data (enforced by the schema).
diff --git a/Engine/Plugins/Runtime/Steam/SteamAudio/Source/SteamAudio/Private/PhononProbeVolume.h b/Engine/Plugins/Runtime/Steam/SteamAudio/Source/SteamAudio/Private/PhononProbeVolume.h
--- a/Engine/Plugins/Runtime/Steam/SteamAudio/Source/SteamAudio/Private/PhononProbeVolume.h
+++ b/Engine/Plugins/Runtime/Steam/SteamAudio/Source/SteamAudio/Private/PhononProbeVolume.h
@@ -124,7 +124,7 @@ public:
 
 	// Component used for visualization.
 	UPROPERTY()
-	class UPhononProbeComponent* PhononProbeComponent;
+	TObjectPtr<class UPhononProbeComponent> PhononProbeComponent;
 
 	// Current filename where probe box data is stored. Used to maintain connection if volume is renamed.
 	UPROPERTY()
diff --git a/Engine/Plugins/Runtime/WaveTable/Source/WaveTable/Public/WaveTableTransform.h b/Engine/Plugins/Runtime/WaveTable/Source/WaveTable/Public/WaveTableTransform.h
--- a/Engine/Plugins/Runtime/WaveTable/Source/WaveTable/Public/WaveTableTransform.h
+++ b/Engine/Plugins/Runtime/WaveTable/Source/WaveTable/Public/WaveTableTransform.h
@@ -38,7 +38,7 @@ struct WAVETABLE_API FWaveTableTransform
 
 	/** Asset curve reference to apply if output curve type is set to 'Shared.' */
 	UPROPERTY(EditAnywhere, meta = (DisplayName = "Asset"), Category = Curve, BlueprintReadWrite)
-	UCurveFloat* CurveShared = nullptr;
+	TObjectPtr<UCurveFloat> CurveShared = nullptr;
 
 	UPROPERTY()
 	TArray<float> WaveTable;
diff --git a/Engine/Plugins/Runtime/XR/XRVisualization/Source/XRVisualization/Public/XRVisualizationFunctionLibrary.h b/Engine/Plugins/Runtime/XR/XRVisualization/Source/XRVisualization/Public/XRVisualizationFunctionLibrary.h
--- a/Engine/Plugins/Runtime/XR/XRVisualization/Source/XRVisualization/Public/XRVisualizationFunctionLibrary.h
+++ b/Engine/Plugins/Runtime/XR/XRVisualization/Source/XRVisualization/Public/XRVisualizationFunctionLibrary.h
@@ -18,13 +18,13 @@ class UXRVisualizationLoadHelper : public UObject
 
 public:
 	UPROPERTY()
-	class UStaticMesh* GenericHMD;
+	TObjectPtr<class UStaticMesh> GenericHMD;
 	UPROPERTY()
-	class UStaticMesh* OculusControllerMesh;
+	TObjectPtr<class UStaticMesh> OculusControllerMesh;
 	UPROPERTY()
-	class UStaticMesh* ViveControllerMesh;
+	TObjectPtr<class UStaticMesh> ViveControllerMesh;
 	UPROPERTY()
-	class UStaticMesh* STEMControllerMesh;
+	TObjectPtr<class UStaticMesh> STEMControllerMesh;
 };
 
 UCLASS()
@@ -56,5 +56,5 @@ private:
 	static void RenderFinger(const FXRMotionControllerData& XRData, int32 FingerStart, int32 FingerEnd);
 
 	UPROPERTY()
-	UXRVisualizationLoadHelper* LoadHelper;
+	TObjectPtr<UXRVisualizationLoadHelper> LoadHelper;
 };
diff --git a/Engine/Plugins/Runtime/ZoneGraph/Source/ZoneGraph/Public/ZoneGraphBuilder.h b/Engine/Plugins/Runtime/ZoneGraph/Source/ZoneGraph/Public/ZoneGraphBuilder.h
--- a/Engine/Plugins/Runtime/ZoneGraph/Source/ZoneGraph/Public/ZoneGraphBuilder.h
+++ b/Engine/Plugins/Runtime/ZoneGraph/Source/ZoneGraph/Public/ZoneGraphBuilder.h
@@ -22,7 +22,7 @@ struct ZONEGRAPH_API FZoneGraphBuilderRegisteredComponent
 	FZoneGraphBuilderRegisteredComponent(UZoneShapeComponent* InComponent) : Component(InComponent) {}
 
 	UPROPERTY()
-	UZoneShapeComponent* Component = nullptr;
+	TObjectPtr<UZoneShapeComponent> Component = nullptr;
 
 	UPROPERTY()
 	uint32 ShapeHash = 0;
@@ -55,7 +55,7 @@ struct ZONEGRAPH_API FZoneGraphBuildData
 	}
 
 	UPROPERTY()
-	TMap<const UZoneShapeComponent*, FZoneShapeComponentBuildData> ZoneShapeComponentBuildData;
+	TMap<TObjectPtr<const UZoneShapeComponent>, FZoneShapeComponentBuildData> ZoneShapeComponentBuildData;
 };
 
 USTRUCT()
@@ -101,7 +101,7 @@ protected:
 	TArray<int32> ShapeComponentsFreeList;
 
 	UPROPERTY(Transient)
-	TMap<UZoneShapeComponent*, int32> ShapeComponentToIndex;
+	TMap<TObjectPtr<UZoneShapeComponent>, int32> ShapeComponentToIndex;
 
 	UPROPERTY(Transient)
 	FZoneGraphBuildData BuildData;
diff --git a/Engine/Plugins/Runtime/ZoneGraph/Source/ZoneGraph/Public/ZoneGraphData.h b/Engine/Plugins/Runtime/ZoneGraph/Source/ZoneGraph/Public/ZoneGraphData.h
--- a/Engine/Plugins/Runtime/ZoneGraph/Source/ZoneGraph/Public/ZoneGraphData.h
+++ b/Engine/Plugins/Runtime/ZoneGraph/Source/ZoneGraph/Public/ZoneGraphData.h
@@ -61,7 +61,7 @@ protected:
 	bool bEnableDrawing;
 
 	UPROPERTY(transient, duplicatetransient)
-	UZoneGraphRenderingComponent* RenderingComp;
+	TObjectPtr<UZoneGraphRenderingComponent> RenderingComp;
 
 	UPROPERTY()
 	FZoneGraphStorage ZoneStorage;
diff --git a/Engine/Plugins/Runtime/ZoneGraph/Source/ZoneGraph/Public/ZoneGraphSubsystem.h b/Engine/Plugins/Runtime/ZoneGraph/Source/ZoneGraph/Public/ZoneGraphSubsystem.h
--- a/Engine/Plugins/Runtime/ZoneGraph/Source/ZoneGraph/Public/ZoneGraphSubsystem.h
+++ b/Engine/Plugins/Runtime/ZoneGraph/Source/ZoneGraph/Public/ZoneGraphSubsystem.h
@@ -28,7 +28,7 @@ struct FRegisteredZoneGraphData
 	}
 
 	UPROPERTY()
-	AZoneGraphData* ZoneGraphData = nullptr;
+	TObjectPtr<AZoneGraphData> ZoneGraphData = nullptr;
 
 	int32 Generation = 1;	// Starting at generation 1 so that 0 can be invalid.
 	bool bInUse = false;	// Extra bit indicating that the data is in meant to be in use. This tried to capture the case where ZoneGraphData might get nullified without notifying.
diff --git a/Engine/Plugins/Runtime/ZoneGraph/Source/ZoneGraph/Public/ZoneShapeActor.h b/Engine/Plugins/Runtime/ZoneGraph/Source/ZoneGraph/Public/ZoneShapeActor.h
--- a/Engine/Plugins/Runtime/ZoneGraph/Source/ZoneGraph/Public/ZoneShapeActor.h
+++ b/Engine/Plugins/Runtime/ZoneGraph/Source/ZoneGraph/Public/ZoneShapeActor.h
@@ -26,5 +26,5 @@ public:
 protected:
 
 	UPROPERTY(Category = Zone, VisibleAnywhere, BlueprintReadOnly, meta = (AllowPrivateAccess = "true"))
-	UZoneShapeComponent* ShapeComponent;
+	TObjectPtr<UZoneShapeComponent> ShapeComponent;
 };
diff --git a/Engine/Plugins/Runtime/ZoneGraph/Source/ZoneGraphDebug/Public/ZoneGraphTestingActor.h b/Engine/Plugins/Runtime/ZoneGraph/Source/ZoneGraphDebug/Public/ZoneGraphTestingActor.h
--- a/Engine/Plugins/Runtime/ZoneGraph/Source/ZoneGraphDebug/Public/ZoneGraphTestingActor.h
+++ b/Engine/Plugins/Runtime/ZoneGraph/Source/ZoneGraphDebug/Public/ZoneGraphTestingActor.h
@@ -31,7 +31,7 @@ protected:
 	virtual void OnOwnerSet() {}
 
 	UPROPERTY()
-	UZoneGraphTestingComponent* OwnerComponent;
+	TObjectPtr<UZoneGraphTestingComponent> OwnerComponent;
 };
 
 /** Actor for testing ZoneGraph functionality. */
@@ -84,7 +84,7 @@ protected:
 	FDelegateHandle OnDataRemovedHandle;
 
 	UPROPERTY(Transient)
-	UZoneGraphSubsystem* ZoneGraph;
+	TObjectPtr<UZoneGraphSubsystem> ZoneGraph;
 
 	UPROPERTY(Transient)
 	FZoneGraphLaneLocation LaneLocation;
@@ -124,13 +124,13 @@ protected:
 	bool bDrawLanePath = false;
 
 	UPROPERTY(EditAnywhere, Category = Test)
-	AZoneGraphTestingActor* OtherActor;
+	TObjectPtr<AZoneGraphTestingActor> OtherActor;
 
 	TArray<FZoneGraphLinkedLane> LinkedLanes;
 	FZoneGraphLanePath LanePath;
 
 	UPROPERTY(EditAnywhere, Category = Test, Instanced)
-	TArray<UZoneLaneTest*> CustomTests;
+	TArray<TObjectPtr<UZoneLaneTest>> CustomTests;
 
 private:
 	void ExecuteOnEachCustomTest(TFunctionRef<void(UZoneLaneTest&)> ExecFunc);
@@ -161,5 +161,5 @@ public:
 
 protected:
 	UPROPERTY(Category = Default, VisibleAnywhere, meta = (AllowPrivateAccess = "true"))
-	UZoneGraphTestingComponent* DebugComp;
+	TObjectPtr<UZoneGraphTestingComponent> DebugComp;
 };
diff --git a/Engine/Plugins/Runtime/ZoneGraphAnnotations/Source/ZoneGraphAnnotations/Public/Annotations/ZoneGraphDisturbanceAnnotation.h b/Engine/Plugins/Runtime/ZoneGraphAnnotations/Source/ZoneGraphAnnotations/Public/Annotations/ZoneGraphDisturbanceAnnotation.h
--- a/Engine/Plugins/Runtime/ZoneGraphAnnotations/Source/ZoneGraphAnnotations/Public/Annotations/ZoneGraphDisturbanceAnnotation.h
+++ b/Engine/Plugins/Runtime/ZoneGraphAnnotations/Source/ZoneGraphAnnotations/Public/Annotations/ZoneGraphDisturbanceAnnotation.h
@@ -225,7 +225,7 @@ protected:
 
 	/** Cached ZoneGraphSubsystem */
 	UPROPERTY(Transient)
-	UZoneGraphSubsystem* ZoneGraphSubsystem = nullptr;
+	TObjectPtr<UZoneGraphSubsystem> ZoneGraphSubsystem = nullptr;
 };
 
 /**
diff --git a/Engine/Plugins/Runtime/ZoneGraphAnnotations/Source/ZoneGraphAnnotations/Public/ZoneGraphAnnotationSubsystem.h b/Engine/Plugins/Runtime/ZoneGraphAnnotations/Source/ZoneGraphAnnotations/Public/ZoneGraphAnnotationSubsystem.h
--- a/Engine/Plugins/Runtime/ZoneGraphAnnotations/Source/ZoneGraphAnnotations/Public/ZoneGraphAnnotationSubsystem.h
+++ b/Engine/Plugins/Runtime/ZoneGraphAnnotations/Source/ZoneGraphAnnotations/Public/ZoneGraphAnnotationSubsystem.h
@@ -55,7 +55,7 @@ struct FRegisteredZoneGraphAnnotation
 	}
 
 	UPROPERTY()
-	UZoneGraphAnnotationComponent* AnnotationComponent = nullptr;
+	TObjectPtr<UZoneGraphAnnotationComponent> AnnotationComponent = nullptr;
 
 	FZoneGraphTagMask AnnotationTags = FZoneGraphTagMask::None;	// Combination of all registered Annotation tag masks.
 };
@@ -133,7 +133,7 @@ protected:
 
 	/** Lookup table from tag index to Annotation */
 	UPROPERTY(Transient)
-	TArray<UZoneGraphAnnotationComponent*> TagToAnnotationLookup;
+	TArray<TObjectPtr<UZoneGraphAnnotationComponent>> TagToAnnotationLookup;
 	
 	/** Combined tags for each ZoneGraphData. Each ZoneGraphData is indexed by it's data handle index, so there can be gaps in the array. */
 	FZoneGraphAnnotationTagContainer AnnotationTagContainer;
diff --git a/Engine/Plugins/Runtime/ZoneGraphAnnotations/Source/ZoneGraphAnnotations/Public/ZoneGraphAnnotationTestingActor.h b/Engine/Plugins/Runtime/ZoneGraphAnnotations/Source/ZoneGraphAnnotations/Public/ZoneGraphAnnotationTestingActor.h
--- a/Engine/Plugins/Runtime/ZoneGraphAnnotations/Source/ZoneGraphAnnotations/Public/ZoneGraphAnnotationTestingActor.h
+++ b/Engine/Plugins/Runtime/ZoneGraphAnnotations/Source/ZoneGraphAnnotations/Public/ZoneGraphAnnotationTestingActor.h
@@ -32,7 +32,7 @@ protected:
 	virtual void OnOwnerSet() {}
 
 	UPROPERTY()
-	UZoneGraphAnnotationTestingComponent* OwnerComponent;
+	TObjectPtr<UZoneGraphAnnotationTestingComponent> OwnerComponent;
 };
 
 
@@ -67,7 +67,7 @@ protected:
 #endif
 
 	UPROPERTY(EditAnywhere, Category = "Test", Instanced)
-	TArray<UZoneGraphAnnotationTest*> Tests;
+	TArray<TObjectPtr<UZoneGraphAnnotationTest>> Tests;
 };
 
 
@@ -92,5 +92,5 @@ protected:
 #endif
 
 	UPROPERTY(Category = Default, VisibleAnywhere, meta = (AllowPrivateAccess = "true"))
-	UZoneGraphAnnotationTestingComponent* TestingComp;
+	TObjectPtr<UZoneGraphAnnotationTestingComponent> TestingComp;
 };
diff --git a/Engine/Plugins/Runtime/nDisplay/Source/DisplayCluster/Public/Blueprints/DisplayClusterBlueprint.h b/Engine/Plugins/Runtime/nDisplay/Source/DisplayCluster/Public/Blueprints/DisplayClusterBlueprint.h
--- a/Engine/Plugins/Runtime/nDisplay/Source/DisplayCluster/Public/Blueprints/DisplayClusterBlueprint.h
+++ b/Engine/Plugins/Runtime/nDisplay/Source/DisplayCluster/Public/Blueprints/DisplayClusterBlueprint.h
@@ -68,7 +68,7 @@ private:
 
 protected:
 	UPROPERTY()
-	UDisplayClusterConfigurationData* ConfigData;
+	TObjectPtr<UDisplayClusterConfigurationData> ConfigData;
 
 private:
 	friend class FDisplayClusterConfiguratorVersionUtils;
diff --git a/Engine/Plugins/Runtime/nDisplay/Source/DisplayCluster/Public/Components/DisplayClusterCameraComponent.h b/Engine/Plugins/Runtime/nDisplay/Source/DisplayCluster/Public/Components/DisplayClusterCameraComponent.h
--- a/Engine/Plugins/Runtime/nDisplay/Source/DisplayCluster/Public/Components/DisplayClusterCameraComponent.h
+++ b/Engine/Plugins/Runtime/nDisplay/Source/DisplayCluster/Public/Components/DisplayClusterCameraComponent.h
@@ -151,10 +151,10 @@ protected:
 	float GizmoScaleMultiplier;
 
 	UPROPERTY(Transient)
-	UBillboardComponent* SpriteComponent;
+	TObjectPtr<UBillboardComponent> SpriteComponent;
 
 	UPROPERTY()
-	UTexture2D* SpriteTexture;
+	TObjectPtr<UTexture2D> SpriteTexture;
 #endif
 
 private:
diff --git a/Engine/Plugins/Runtime/nDisplay/Source/DisplayCluster/Public/Components/DisplayClusterOriginComponent.h b/Engine/Plugins/Runtime/nDisplay/Source/DisplayCluster/Public/Components/DisplayClusterOriginComponent.h
--- a/Engine/Plugins/Runtime/nDisplay/Source/DisplayCluster/Public/Components/DisplayClusterOriginComponent.h
+++ b/Engine/Plugins/Runtime/nDisplay/Source/DisplayCluster/Public/Components/DisplayClusterOriginComponent.h
@@ -65,10 +65,10 @@ protected:
 
 	/** Proxy mesh to render */
 	UPROPERTY(BlueprintReadWrite, EditAnywhere, Category = "Gizmo", meta = (EditCondition = "bEnableGizmo"))
-	UStaticMesh* ProxyMesh;
+	TObjectPtr<UStaticMesh> ProxyMesh;
 
 	/** Proxy mesh component */
 	UPROPERTY(Transient)
-	UStaticMeshComponent* ProxyMeshComponent;
+	TObjectPtr<UStaticMeshComponent> ProxyMeshComponent;
 #endif
 };
diff --git a/Engine/Plugins/Runtime/nDisplay/Source/DisplayCluster/Public/Components/DisplayClusterPreviewComponent.h b/Engine/Plugins/Runtime/nDisplay/Source/DisplayCluster/Public/Components/DisplayClusterPreviewComponent.h
--- a/Engine/Plugins/Runtime/nDisplay/Source/DisplayCluster/Public/Components/DisplayClusterPreviewComponent.h
+++ b/Engine/Plugins/Runtime/nDisplay/Source/DisplayCluster/Public/Components/DisplayClusterPreviewComponent.h
@@ -121,11 +121,11 @@ protected:
 protected:
 	// Texture for preview material
 	UPROPERTY(VisibleAnywhere, BlueprintReadOnly, Category = "Preview", meta = (DisplayName = "Render Target"))
-	UTextureRenderTarget2D* RenderTarget;
+	TObjectPtr<UTextureRenderTarget2D> RenderTarget;
 
 	// Texture when DCRA has post process disabled but is requesting a post process render target.
 	UPROPERTY(VisibleAnywhere, BlueprintReadOnly, Category = "Preview")
-	UTextureRenderTarget2D* RenderTargetPostProcess;
+	TObjectPtr<UTextureRenderTarget2D> RenderTargetPostProcess;
 	
 private:
 	// Saved mesh policy params
@@ -133,7 +133,7 @@ private:
 	FDisplayClusterConfigurationProjection WarpMeshSavedProjectionPolicy;
 
 	UPROPERTY()
-	ADisplayClusterRootActor* RootActor = nullptr;
+	TObjectPtr<ADisplayClusterRootActor> RootActor = nullptr;
 
 	UPROPERTY()
 	FString ViewportId;
@@ -142,25 +142,25 @@ private:
 	FString ClusterNodeId;
 
 	UPROPERTY()
-	UDisplayClusterConfigurationViewport* ViewportConfig = nullptr;
+	TObjectPtr<UDisplayClusterConfigurationViewport> ViewportConfig = nullptr;
 
 	UPROPERTY()
-	UMeshComponent* PreviewMesh = nullptr;
+	TObjectPtr<UMeshComponent> PreviewMesh = nullptr;
 
 	UPROPERTY()
 	bool bIsRootActorPreviewMesh = false;
 
 	UPROPERTY()
-	UMaterial* OriginalMaterial = nullptr;
+	TObjectPtr<UMaterial> OriginalMaterial = nullptr;
 
 	UPROPERTY()
-	UMaterial* PreviewMaterial = nullptr;
+	TObjectPtr<UMaterial> PreviewMaterial = nullptr;
 
 	UPROPERTY()
-	UMaterialInstanceDynamic* PreviewMaterialInstance = nullptr;
+	TObjectPtr<UMaterialInstanceDynamic> PreviewMaterialInstance = nullptr;
 
 	UPROPERTY(Transient)
-	UTexture* OverrideTexture = nullptr;
+	TObjectPtr<UTexture> OverrideTexture = nullptr;
 
 #endif /*WITH_EDITORONLY_DATA*/
 };
diff --git a/Engine/Plugins/Runtime/nDisplay/Source/DisplayCluster/Public/Components/DisplayClusterXformComponent.h b/Engine/Plugins/Runtime/nDisplay/Source/DisplayCluster/Public/Components/DisplayClusterXformComponent.h
--- a/Engine/Plugins/Runtime/nDisplay/Source/DisplayCluster/Public/Components/DisplayClusterXformComponent.h
+++ b/Engine/Plugins/Runtime/nDisplay/Source/DisplayCluster/Public/Components/DisplayClusterXformComponent.h
@@ -65,10 +65,10 @@ protected:
 
 	/** Proxy mesh to render */
 	UPROPERTY(BlueprintReadWrite, EditAnywhere, Category = "Gizmo", meta = (EditCondition = "bEnableGizmo"))
-	UStaticMesh* ProxyMesh;
+	TObjectPtr<UStaticMesh> ProxyMesh;
 
 	/** Proxy mesh component */
 	UPROPERTY(Transient)
-	UStaticMeshComponent* ProxyMeshComponent;
+	TObjectPtr<UStaticMeshComponent> ProxyMeshComponent;
 #endif
 };
diff --git a/Engine/Plugins/Runtime/nDisplay/Source/DisplayCluster/Public/DisplayClusterGameEngine.h b/Engine/Plugins/Runtime/nDisplay/Source/DisplayCluster/Public/DisplayClusterGameEngine.h
--- a/Engine/Plugins/Runtime/nDisplay/Source/DisplayCluster/Public/DisplayClusterGameEngine.h
+++ b/Engine/Plugins/Runtime/nDisplay/Source/DisplayCluster/Public/DisplayClusterGameEngine.h
@@ -27,7 +27,7 @@ class DISPLAYCLUSTER_API UDisplayClusterGameEngine
 	GENERATED_BODY()
 
 	UPROPERTY()
-	UObject* WorldContextObject = nullptr;
+	TObjectPtr<UObject> WorldContextObject = nullptr;
 
 public:
 	virtual void Init(class IEngineLoop* InEngineLoop) override;
diff --git a/Engine/Plugins/Runtime/nDisplay/Source/DisplayCluster/Public/DisplayClusterRootActor.h b/Engine/Plugins/Runtime/nDisplay/Source/DisplayCluster/Public/DisplayClusterRootActor.h
--- a/Engine/Plugins/Runtime/nDisplay/Source/DisplayCluster/Public/DisplayClusterRootActor.h
+++ b/Engine/Plugins/Runtime/nDisplay/Source/DisplayCluster/Public/DisplayClusterRootActor.h
@@ -244,19 +244,19 @@ private:
 	 * nDisplay details panel will hide this from actually being visible.
 	 */
 	UPROPERTY(EditAnywhere, Category = "NDisplay", meta = (HideProperty))
-	USceneComponent* DisplayClusterRootComponent;
+	TObjectPtr<USceneComponent> DisplayClusterRootComponent;
 
 	/**
 	 * Default camera component. It's an outer camera in VP/ICVFX terminology. Always exists on a DCRA instance.
 	 */
 	UPROPERTY(VisibleAnywhere, BlueprintGetter=GetDefaultCamera, Category = "NDisplay|Components")
-	UDisplayClusterCameraComponent* DefaultViewPoint;
+	TObjectPtr<UDisplayClusterCameraComponent> DefaultViewPoint;
 
 	/**
 	 * Helper sync component. Performs sync procedure during Tick phase.
 	 */
 	UPROPERTY()
-	UDisplayClusterSyncTickComponent* SyncTickComponent;
+	TObjectPtr<UDisplayClusterSyncTickComponent> SyncTickComponent;
 
 private:
 	// Current operation mode
@@ -277,7 +277,7 @@ public:
 	 * If set from the DisplayCluster BP Compiler it will be loaded from the class default subobjects in run-time.
 	 */
 	UPROPERTY(VisibleAnywhere, BlueprintReadOnly, Instanced, Category = "NDisplay", meta = (AllowPrivateAccess = "true"))
-	UDisplayClusterConfigurationData* CurrentConfigData;
+	TObjectPtr<UDisplayClusterConfigurationData> CurrentConfigData;
 
 public:
 	// UObject interface
@@ -346,7 +346,7 @@ public:
 
 private:
 	UPROPERTY(Transient)
-	TMap<FString, UDisplayClusterPreviewComponent*> PreviewComponents;
+	TMap<FString, TObjectPtr<UDisplayClusterPreviewComponent>> PreviewComponents;
 
 	UPROPERTY(Transient)
 	bool bDeferPreviewGeneration;
diff --git a/Engine/Plugins/Runtime/nDisplay/Source/DisplayCluster/Public/DisplayClusterTestPatternsActor.h b/Engine/Plugins/Runtime/nDisplay/Source/DisplayCluster/Public/DisplayClusterTestPatternsActor.h
--- a/Engine/Plugins/Runtime/nDisplay/Source/DisplayCluster/Public/DisplayClusterTestPatternsActor.h
+++ b/Engine/Plugins/Runtime/nDisplay/Source/DisplayCluster/Public/DisplayClusterTestPatternsActor.h
@@ -64,10 +64,10 @@ protected:
 protected:
 	/** Postprocess component */
 	UPROPERTY(VisibleAnywhere, BlueprintReadOnly, Category = "NDisplay")
-	UPostProcessComponent* PostProcessComponent;
+	TObjectPtr<UPostProcessComponent> PostProcessComponent;
 
 	UPROPERTY(EditAnywhere, BlueprintReadWrite, Category = "NDisplay")
-	TMap<FString, UMaterial*> CalibrationPatterns;
+	TMap<FString, TObjectPtr<UMaterial>> CalibrationPatterns;
 
 private:
 	void InitializeMaterials();
diff --git a/Engine/Plugins/Runtime/nDisplay/Source/DisplayClusterConfiguration/Public/DisplayClusterConfigurationTypes.h b/Engine/Plugins/Runtime/nDisplay/Source/DisplayClusterConfiguration/Public/DisplayClusterConfigurationTypes.h
--- a/Engine/Plugins/Runtime/nDisplay/Source/DisplayClusterConfiguration/Public/DisplayClusterConfigurationTypes.h
+++ b/Engine/Plugins/Runtime/nDisplay/Source/DisplayClusterConfiguration/Public/DisplayClusterConfigurationTypes.h
@@ -146,13 +146,13 @@ class DISPLAYCLUSTERCONFIGURATION_API UDisplayClusterConfigurationScene
 
 public:
 	UPROPERTY()
-	TMap<FString, UDisplayClusterConfigurationSceneComponentXform*> Xforms;
+	TMap<FString, TObjectPtr<UDisplayClusterConfigurationSceneComponentXform>> Xforms;
 
 	UPROPERTY()
-	TMap<FString, UDisplayClusterConfigurationSceneComponentScreen*> Screens;
+	TMap<FString, TObjectPtr<UDisplayClusterConfigurationSceneComponentScreen>> Screens;
 
 	UPROPERTY()
-	TMap<FString, UDisplayClusterConfigurationSceneComponentCamera*> Cameras;
+	TMap<FString, TObjectPtr<UDisplayClusterConfigurationSceneComponentCamera>> Cameras;
 
 protected:
 	virtual void GetObjectsToExport(TArray<UObject*>& OutObjects) override;
@@ -345,7 +345,7 @@ public:
 #endif
 
 	UPROPERTY(VisibleInstanceOnly, BlueprintReadOnly, EditFixedSize, Instanced, Category = "Configuration", meta = (DisplayThumbnail = false))
-	TMap<FString, UDisplayClusterConfigurationViewport*> Viewports;
+	TMap<FString, TObjectPtr<UDisplayClusterConfigurationViewport>> Viewports;
 
 	UPROPERTY(EditDefaultsOnly, BlueprintReadOnly, Category = "Configuration", meta = (DisplayName = "Custom Output Settings", DisplayThumbnail = false, ShowInnerProperties))
 	TMap<FString, FDisplayClusterConfigurationPostprocess> Postprocess;
@@ -446,11 +446,11 @@ public:
 	FDisplayClusterConfigurationFailoverSettings Failover;
 
 	UPROPERTY(VisibleInstanceOnly, BlueprintReadOnly, EditFixedSize, Instanced, Category = NDisplay, meta = (DisplayThumbnail = false, HideProperty))
-	TMap<FString, UDisplayClusterConfigurationClusterNode*> Nodes;
+	TMap<FString, TObjectPtr<UDisplayClusterConfigurationClusterNode>> Nodes;
 
 #if WITH_EDITORONLY_DATA
 	UPROPERTY(Instanced)
-	TMap<FString, UDisplayClusterConfigurationHostDisplayData*> HostDisplayData;
+	TMap<FString, TObjectPtr<UDisplayClusterConfigurationHostDisplayData>> HostDisplayData;
 #endif
 
 protected:
@@ -540,10 +540,10 @@ public:
 	FDisplayClusterConfigurationInfo Info;
 
 	UPROPERTY(VisibleInstanceOnly, BlueprintReadOnly, Transient, Category = NDisplay, meta = (HideProperty))
-	UDisplayClusterConfigurationScene* Scene;
+	TObjectPtr<UDisplayClusterConfigurationScene> Scene;
 
 	UPROPERTY(VisibleInstanceOnly, BlueprintReadOnly, Instanced, Category = NDisplay)
-	UDisplayClusterConfigurationCluster* Cluster;
+	TObjectPtr<UDisplayClusterConfigurationCluster> Cluster;
 
 	UPROPERTY(EditDefaultsOnly, BlueprintReadWrite, Category = Configuration, meta = (DisplayAfter = "Diagnostics"))
 	TMap<FString, FString> CustomParameters;
diff --git a/Engine/Plugins/Runtime/nDisplay/Source/DisplayClusterConfiguration/Public/DisplayClusterConfigurationTypes_Base.h b/Engine/Plugins/Runtime/nDisplay/Source/DisplayClusterConfiguration/Public/DisplayClusterConfigurationTypes_Base.h
--- a/Engine/Plugins/Runtime/nDisplay/Source/DisplayClusterConfiguration/Public/DisplayClusterConfigurationTypes_Base.h
+++ b/Engine/Plugins/Runtime/nDisplay/Source/DisplayClusterConfiguration/Public/DisplayClusterConfigurationTypes_Base.h
@@ -135,7 +135,7 @@ protected:
 
 private:
 	UPROPERTY(Export)
-	TArray<UObject*> ExportedObjects;
+	TArray<TObjectPtr<UObject>> ExportedObjects;
 };
 
 USTRUCT(Blueprintable)
diff --git a/Engine/Plugins/Runtime/nDisplay/Source/DisplayClusterConfiguration/Public/DisplayClusterConfigurationTypes_ICVFX.h b/Engine/Plugins/Runtime/nDisplay/Source/DisplayClusterConfiguration/Public/DisplayClusterConfigurationTypes_ICVFX.h
--- a/Engine/Plugins/Runtime/nDisplay/Source/DisplayClusterConfiguration/Public/DisplayClusterConfigurationTypes_ICVFX.h
+++ b/Engine/Plugins/Runtime/nDisplay/Source/DisplayClusterConfiguration/Public/DisplayClusterConfigurationTypes_ICVFX.h
@@ -161,7 +161,7 @@ public:
 
 	/** Texture to use as the chromakey marker tile. */
 	UPROPERTY(EditAnywhere, BlueprintReadWrite, Category = NDisplay, meta = (EditCondition = "bEnable"))
-	UTexture* MarkerTileRGBA = nullptr;
+	TObjectPtr<UTexture> MarkerTileRGBA = nullptr;
 
 	/** Scale value for the size of each chromakey marker tile. */
 	UPROPERTY(EditAnywhere, BlueprintReadWrite, Category = NDisplay, meta = (EditCondition = "bEnable", ClampMin = "0", UIMin = "0", DisplayName = "Marker Scale"))
diff --git a/Engine/Plugins/Runtime/nDisplay/Source/DisplayClusterConfiguration/Public/DisplayClusterConfigurationTypes_OutputRemap.h b/Engine/Plugins/Runtime/nDisplay/Source/DisplayClusterConfiguration/Public/DisplayClusterConfigurationTypes_OutputRemap.h
--- a/Engine/Plugins/Runtime/nDisplay/Source/DisplayClusterConfiguration/Public/DisplayClusterConfigurationTypes_OutputRemap.h
+++ b/Engine/Plugins/Runtime/nDisplay/Source/DisplayClusterConfiguration/Public/DisplayClusterConfigurationTypes_OutputRemap.h
@@ -42,7 +42,7 @@ public:
 
 	/** The Static Mesh reference to use for output remapping when the Data Source is set to Static Mesh */
 	UPROPERTY(EditAnywhere, BlueprintReadOnly, Category = "NDisplay OutputRemap", meta = (EditCondition = "DataSource == EDisplayClusterConfigurationFramePostProcess_OutputRemapSource::StaticMesh && bEnable"))
-	class UStaticMesh* StaticMesh = nullptr;
+	TObjectPtr<class UStaticMesh> StaticMesh = nullptr;
 
 	/** The MeshComponent reference (ProceduralMeshComponent or StaticMeshComponent) to use for output remapping when the Data Source is set to Mesh Component */
 	UPROPERTY(EditAnywhere, BlueprintReadOnly, Category = "NDisplay OutputRemap", meta = (EditCondition = "DataSource == EDisplayClusterConfigurationFramePostProcess_OutputRemapSource::MeshComponent && bEnable"))
diff --git a/Engine/Plugins/Runtime/nDisplay/Source/DisplayClusterConfiguration/Public/DisplayClusterConfigurationTypes_PostRender.h b/Engine/Plugins/Runtime/nDisplay/Source/DisplayClusterConfiguration/Public/DisplayClusterConfigurationTypes_PostRender.h
--- a/Engine/Plugins/Runtime/nDisplay/Source/DisplayClusterConfiguration/Public/DisplayClusterConfigurationTypes_PostRender.h
+++ b/Engine/Plugins/Runtime/nDisplay/Source/DisplayClusterConfiguration/Public/DisplayClusterConfigurationTypes_PostRender.h
@@ -27,7 +27,7 @@ public:
 
 	/** Texture to use in place of the inner frustum. */
 	UPROPERTY(EditAnywhere, BlueprintReadWrite, Category = "NDisplay Render", meta = (EditCondition = "bAllowReplace"))
-	UTexture* SourceTexture = nullptr;
+	TObjectPtr<UTexture> SourceTexture = nullptr;
 
 	/** Set to True to crop the texture for the inner frustum as specified below. */
 	UPROPERTY(EditAnywhere, BlueprintReadWrite, Category = "NDisplay Render", meta = (DisplayName = "Use Texture Crop", EditCondition = "bAllowReplace"))
diff --git a/Engine/Plugins/Runtime/nDisplay/Source/DisplayClusterConfigurator/Private/DisplayClusterConfiguratorFactory.h b/Engine/Plugins/Runtime/nDisplay/Source/DisplayClusterConfigurator/Private/DisplayClusterConfiguratorFactory.h
--- a/Engine/Plugins/Runtime/nDisplay/Source/DisplayClusterConfigurator/Private/DisplayClusterConfiguratorFactory.h
+++ b/Engine/Plugins/Runtime/nDisplay/Source/DisplayClusterConfigurator/Private/DisplayClusterConfiguratorFactory.h
@@ -48,7 +48,7 @@ private:
 	TSubclassOf<class ADisplayClusterRootActor> ParentClass;
 
 	UPROPERTY(Transient)
-	UDisplayClusterBlueprint* BlueprintToCopy;
+	TObjectPtr<UDisplayClusterBlueprint> BlueprintToCopy;
 
 	FDelegateHandle OnConfigureNewAssetRequestHandle;
 	bool bIsConfigureNewAssetRequest;
diff --git a/Engine/Plugins/Runtime/nDisplay/Source/DisplayClusterConfigurator/Private/DisplayClusterConfiguratorVersionUtils.h b/Engine/Plugins/Runtime/nDisplay/Source/DisplayClusterConfigurator/Private/DisplayClusterConfiguratorVersionUtils.h
--- a/Engine/Plugins/Runtime/nDisplay/Source/DisplayClusterConfigurator/Private/DisplayClusterConfiguratorVersionUtils.h
+++ b/Engine/Plugins/Runtime/nDisplay/Source/DisplayClusterConfigurator/Private/DisplayClusterConfiguratorVersionUtils.h
@@ -28,7 +28,7 @@ public:
 
 public:
 	UPROPERTY(Transient)
-	class UDisplayClusterConfigurationData* nDisplayConfig;
+	TObjectPtr<class UDisplayClusterConfigurationData> nDisplayConfig;
 
 	UPROPERTY()
 	FString PathToConfig;
diff --git a/Engine/Plugins/Runtime/nDisplay/Source/DisplayClusterConfigurator/Private/Views/OutputMapping/EdNodes/DisplayClusterConfiguratorBaseNode.h b/Engine/Plugins/Runtime/nDisplay/Source/DisplayClusterConfigurator/Private/Views/OutputMapping/EdNodes/DisplayClusterConfiguratorBaseNode.h
--- a/Engine/Plugins/Runtime/nDisplay/Source/DisplayClusterConfigurator/Private/Views/OutputMapping/EdNodes/DisplayClusterConfiguratorBaseNode.h
+++ b/Engine/Plugins/Runtime/nDisplay/Source/DisplayClusterConfigurator/Private/Views/OutputMapping/EdNodes/DisplayClusterConfiguratorBaseNode.h
@@ -178,7 +178,7 @@ protected:
 	TWeakObjectPtr<UDisplayClusterConfiguratorBaseNode> Parent;
 
 	UPROPERTY(Transient, NonTransactional)
-	TArray<UDisplayClusterConfiguratorBaseNode*> Children;
+	TArray<TObjectPtr<UDisplayClusterConfiguratorBaseNode>> Children;
 
 	FString NodeName;
 	int32 NodeZIndex;
diff --git a/Engine/Plugins/ScriptPlugin/Source/ScriptEditorPlugin/Classes/ReimportScriptFactory.h b/Engine/Plugins/ScriptPlugin/Source/ScriptEditorPlugin/Classes/ReimportScriptFactory.h
--- a/Engine/Plugins/ScriptPlugin/Source/ScriptEditorPlugin/Classes/ReimportScriptFactory.h
+++ b/Engine/Plugins/ScriptPlugin/Source/ScriptEditorPlugin/Classes/ReimportScriptFactory.h
@@ -20,7 +20,7 @@ class UReimportScriptFactory : public UScriptFactory, public FReimportHandler
 	GENERATED_UCLASS_BODY()
 
 	UPROPERTY()
-	class UScriptBlueprint* OriginalScript;
+	TObjectPtr<class UScriptBlueprint> OriginalScript;
 
 	virtual bool ConfigureProperties() override;
 
diff --git a/Engine/Plugins/ScriptPlugin/Source/ScriptPlugin/Classes/ScriptBlueprintGeneratedClass.h b/Engine/Plugins/ScriptPlugin/Source/ScriptPlugin/Classes/ScriptBlueprintGeneratedClass.h
--- a/Engine/Plugins/ScriptPlugin/Source/ScriptPlugin/Classes/ScriptBlueprintGeneratedClass.h
+++ b/Engine/Plugins/ScriptPlugin/Source/ScriptPlugin/Classes/ScriptBlueprintGeneratedClass.h
@@ -148,7 +148,7 @@ public:
 
 	/** Script-generated properties */
 	UPROPERTY()
-	TArray<UProperty*> ScriptProperties_DEPRECATED;
+	TArray<TObjectPtr<UProperty>> ScriptProperties_DEPRECATED;
 	TArray<TFieldPath<FProperty>> ScriptProperties;
 
 	virtual void PurgeClass(bool bRecompilingOnLoad) override;
diff --git a/Engine/Plugins/VirtualProduction/CameraCalibration/Source/CameraCalibrationEditor/Private/AssetEditor/LensDistortionTool.h b/Engine/Plugins/VirtualProduction/CameraCalibration/Source/CameraCalibrationEditor/Private/AssetEditor/LensDistortionTool.h
--- a/Engine/Plugins/VirtualProduction/CameraCalibration/Source/CameraCalibrationEditor/Private/AssetEditor/LensDistortionTool.h
+++ b/Engine/Plugins/VirtualProduction/CameraCalibration/Source/CameraCalibrationEditor/Private/AssetEditor/LensDistortionTool.h
@@ -64,7 +64,7 @@ private:
 
 	/** The currently selected algorithm */
 	UPROPERTY(Transient)
-	UCameraLensDistortionAlgo* CurrentAlgo;
+	TObjectPtr<UCameraLensDistortionAlgo> CurrentAlgo;
 
 	/** Holds the registered camera lens distortion algos */
 	UPROPERTY(Transient)
diff --git a/Engine/Plugins/VirtualProduction/CameraCalibration/Source/CameraCalibrationEditor/Private/AssetEditor/NodalOffsetTool.h b/Engine/Plugins/VirtualProduction/CameraCalibration/Source/CameraCalibrationEditor/Private/AssetEditor/NodalOffsetTool.h
--- a/Engine/Plugins/VirtualProduction/CameraCalibration/Source/CameraCalibrationEditor/Private/AssetEditor/NodalOffsetTool.h
+++ b/Engine/Plugins/VirtualProduction/CameraCalibration/Source/CameraCalibrationEditor/Private/AssetEditor/NodalOffsetTool.h
@@ -65,7 +65,7 @@ private:
 
 	/** The currently selected nodal offset algorithm */
 	UPROPERTY()
-	UCameraNodalOffsetAlgo* NodalOffsetAlgo;
+	TObjectPtr<UCameraNodalOffsetAlgo> NodalOffsetAlgo;
 
 	/** Map of algo names to overlay MIDs used by those algos */
 	UPROPERTY(Transient)
diff --git a/Engine/Plugins/VirtualProduction/CameraCalibrationCore/Source/CameraCalibrationCore/Public/CameraCalibrationCheckerboard.h b/Engine/Plugins/VirtualProduction/CameraCalibrationCore/Source/CameraCalibrationCore/Public/CameraCalibrationCheckerboard.h
--- a/Engine/Plugins/VirtualProduction/CameraCalibrationCore/Source/CameraCalibrationCore/Public/CameraCalibrationCheckerboard.h
+++ b/Engine/Plugins/VirtualProduction/CameraCalibrationCore/Source/CameraCalibrationCore/Public/CameraCalibrationCheckerboard.h
@@ -41,27 +41,27 @@ public:
 
 	/** Root component, gives the Actor a transform */
 	UPROPERTY()
-	USceneComponent* Root;
+	TObjectPtr<USceneComponent> Root;
 
 	/** TopLeft calibration point */
 	UPROPERTY(EditAnywhere, Category = "Calibration")
-	UCalibrationPointComponent* TopLeft;
+	TObjectPtr<UCalibrationPointComponent> TopLeft;
 
 	/** TopRight calibration point */
 	UPROPERTY(EditAnywhere, Category = "Calibration")
-	UCalibrationPointComponent* TopRight;
+	TObjectPtr<UCalibrationPointComponent> TopRight;
 
 	/** BottomLeft calibration point */
 	UPROPERTY(EditAnywhere, Category = "Calibration")
-	UCalibrationPointComponent* BottomLeft;
+	TObjectPtr<UCalibrationPointComponent> BottomLeft;
 
 	/** BottomRight calibration point */
 	UPROPERTY(EditAnywhere, Category = "Calibration")
-	UCalibrationPointComponent* BottomRight;
+	TObjectPtr<UCalibrationPointComponent> BottomRight;
 
 	/** Center calibration point */
 	UPROPERTY(EditAnywhere, Category = "Calibration")
-	UCalibrationPointComponent* Center;
+	TObjectPtr<UCalibrationPointComponent> Center;
 
 
 	/** Number of rows. It is one of the parameters cv::findChessboardCorners needs */
@@ -82,13 +82,13 @@ public:
 
 	/** The static mesh that we are going to use for all the cubes */
 	UPROPERTY(EditAnywhere, Category = "Calibration")
-	UStaticMesh* CubeMesh;
+	TObjectPtr<UStaticMesh> CubeMesh;
 
 	/** The material that we are going to use for all the odd cubes */
 	UPROPERTY(EditAnywhere, Category = "Calibration")
-	UMaterialInterface* OddCubeMaterial;
+	TObjectPtr<UMaterialInterface> OddCubeMaterial;
 
 	/** The material that we are going to use for all the even cubes */
 	UPROPERTY(EditAnywhere, Category = "Calibration")
-	UMaterialInterface* EvenCubeMaterial;
+	TObjectPtr<UMaterialInterface> EvenCubeMaterial;
 };
diff --git a/Engine/Plugins/VirtualProduction/CameraCalibrationCore/Source/CameraCalibrationCore/Public/CameraCalibrationSubsystem.h b/Engine/Plugins/VirtualProduction/CameraCalibrationCore/Source/CameraCalibrationCore/Public/CameraCalibrationSubsystem.h
--- a/Engine/Plugins/VirtualProduction/CameraCalibrationCore/Source/CameraCalibrationCore/Public/CameraCalibrationSubsystem.h
+++ b/Engine/Plugins/VirtualProduction/CameraCalibrationCore/Source/CameraCalibrationCore/Public/CameraCalibrationSubsystem.h
@@ -136,7 +136,7 @@ public:
 private:
 	/** Default lens file to use when no override has been provided */
 	UPROPERTY(Transient)
-	ULensFile* DefaultLensFile = nullptr;
+	TObjectPtr<ULensFile> DefaultLensFile = nullptr;
 
 	/** Map of model names to ULensModel subclasses */
 	UPROPERTY(Transient)
diff --git a/Engine/Plugins/VirtualProduction/CameraCalibrationCore/Source/CameraCalibrationCore/Public/CameraCalibrationTypes.h b/Engine/Plugins/VirtualProduction/CameraCalibrationCore/Source/CameraCalibrationCore/Public/CameraCalibrationTypes.h
--- a/Engine/Plugins/VirtualProduction/CameraCalibrationCore/Source/CameraCalibrationCore/Public/CameraCalibrationTypes.h
+++ b/Engine/Plugins/VirtualProduction/CameraCalibrationCore/Source/CameraCalibrationCore/Public/CameraCalibrationTypes.h
@@ -25,7 +25,7 @@ PRAGMA_ENABLE_DEPRECATION_WARNINGS
 public:
 	/** CineCameraComponent with which the desired distortion handler is associated */
 	UPROPERTY(BlueprintReadWrite, Category = "Distortion", Transient)
-	UCineCameraComponent* TargetCameraComponent = nullptr;
+	TObjectPtr<UCineCameraComponent> TargetCameraComponent = nullptr;
 
 	/** UObject that produces the distortion state for the desired distortion handler */
 	UPROPERTY(BlueprintReadWrite, Category = "Distortion")
diff --git a/Engine/Plugins/VirtualProduction/CameraCalibrationCore/Source/CameraCalibrationCore/Public/LensComponent.h b/Engine/Plugins/VirtualProduction/CameraCalibrationCore/Source/CameraCalibrationCore/Public/LensComponent.h
--- a/Engine/Plugins/VirtualProduction/CameraCalibrationCore/Source/CameraCalibrationCore/Public/LensComponent.h
+++ b/Engine/Plugins/VirtualProduction/CameraCalibrationCore/Source/CameraCalibrationCore/Public/LensComponent.h
@@ -330,11 +330,11 @@ protected:
 
 	/** Cached MID last applied to the target camera */
 	UPROPERTY()
-	UMaterialInstanceDynamic* LastDistortionMID = nullptr;
+	TObjectPtr<UMaterialInstanceDynamic> LastDistortionMID = nullptr;
 
 	/** Cached most recent target camera, used to clean up the old camera when the user changes the target */
 	UPROPERTY()
-	UCineCameraComponent* LastCameraComponent = nullptr;
+	TObjectPtr<UCineCameraComponent> LastCameraComponent = nullptr;
 
 	/** Map of lens models to handlers */
 	UPROPERTY(Transient)
@@ -363,7 +363,7 @@ protected:
 
 	UE_DEPRECATED(5.1, "This property has been deprecated. Use the LensDistortionHandlerMap to get a handler for the current LensModel.")
 	UPROPERTY(Transient)
-	ULensDistortionModelHandlerBase* LensDistortionHandler_DEPRECATED = nullptr;
+	TObjectPtr<ULensDistortionModelHandlerBase> LensDistortionHandler_DEPRECATED = nullptr;
 
 	UE_DEPRECATED(5.1, "This property has been deprecated. Producer GUIDs are no longer used to identify distortion handlers. Use the LensDistortionHandlerMap to get a handler for the current LensModel.")
 	UPROPERTY(DuplicateTransient)
diff --git a/Engine/Plugins/VirtualProduction/CameraCalibrationCore/Source/CameraCalibrationCore/Public/LensData.h b/Engine/Plugins/VirtualProduction/CameraCalibrationCore/Source/CameraCalibrationCore/Public/LensData.h
--- a/Engine/Plugins/VirtualProduction/CameraCalibrationCore/Source/CameraCalibrationCore/Public/LensData.h
+++ b/Engine/Plugins/VirtualProduction/CameraCalibrationCore/Source/CameraCalibrationCore/Public/LensData.h
@@ -101,7 +101,7 @@ public:
 	 * BA channels are expected to have distortion map (from undistorted (CG) to distorted)
 	 */
 	UPROPERTY(EditAnywhere, BlueprintReadWrite, Category = "Distortion")
-	UTexture* DistortionMap = nullptr;
+	TObjectPtr<UTexture> DistortionMap = nullptr;
 };
 
 /**
diff --git a/Engine/Plugins/VirtualProduction/CameraCalibrationCore/Source/CameraCalibrationCore/Public/LensDistortionModelHandlerBase.h b/Engine/Plugins/VirtualProduction/CameraCalibrationCore/Source/CameraCalibrationCore/Public/LensDistortionModelHandlerBase.h
--- a/Engine/Plugins/VirtualProduction/CameraCalibrationCore/Source/CameraCalibrationCore/Public/LensDistortionModelHandlerBase.h
+++ b/Engine/Plugins/VirtualProduction/CameraCalibrationCore/Source/CameraCalibrationCore/Public/LensDistortionModelHandlerBase.h
@@ -145,7 +145,7 @@ protected:
 
 	/** Dynamically created post-process material instance for the currently specified lens model */
 	UPROPERTY(VisibleAnywhere, BlueprintReadOnly, Category = "Distortion")
-	UMaterialInstanceDynamic* DistortionPostProcessMID = nullptr;
+	TObjectPtr<UMaterialInstanceDynamic> DistortionPostProcessMID = nullptr;
 
 	/** Current state as set by the most recent call to Update() */
 	UPROPERTY(EditAnywhere, BlueprintReadOnly, Category = "Distortion", meta = (ShowOnlyInnerProperties))
@@ -161,19 +161,19 @@ protected:
 
 	/** MID used to draw the undistortion displacement map */
 	UPROPERTY(Transient)
-	UMaterialInstanceDynamic* UndistortionDisplacementMapMID = nullptr;
+	TObjectPtr<UMaterialInstanceDynamic> UndistortionDisplacementMapMID = nullptr;
 
 	/** MID used to draw the distortion displacement map */
 	UPROPERTY(Transient)
-	UMaterialInstanceDynamic* DistortionDisplacementMapMID = nullptr;
+	TObjectPtr<UMaterialInstanceDynamic> DistortionDisplacementMapMID = nullptr;
 
 	/** UV displacement map used to undistort a distorted image */
 	UPROPERTY(Transient)
-	UTextureRenderTarget2D* UndistortionDisplacementMapRT = nullptr;
+	TObjectPtr<UTextureRenderTarget2D> UndistortionDisplacementMapRT = nullptr;
 
 	/** UV displacement map used to distort an undistorted image */
 	UPROPERTY(Transient)
-	UTextureRenderTarget2D* DistortionDisplacementMapRT = nullptr;
+	TObjectPtr<UTextureRenderTarget2D> DistortionDisplacementMapRT = nullptr;
 
 	/** UObject that is producing the distortion state for this handler */
 	UPROPERTY()
diff --git a/Engine/Plugins/VirtualProduction/CameraCalibrationCore/Source/CameraCalibrationCore/Public/LensFile.h b/Engine/Plugins/VirtualProduction/CameraCalibrationCore/Source/CameraCalibrationCore/Public/LensFile.h
--- a/Engine/Plugins/VirtualProduction/CameraCalibrationCore/Source/CameraCalibrationCore/Public/LensFile.h
+++ b/Engine/Plugins/VirtualProduction/CameraCalibrationCore/Source/CameraCalibrationCore/Public/LensFile.h
@@ -288,11 +288,11 @@ protected:
 
 	/** Texture used to store temporary undistortion displacement map when using map blending */
 	UPROPERTY(Transient)
-	TArray<UTextureRenderTarget2D*> UndistortionDisplacementMapHolders;
+	TArray<TObjectPtr<UTextureRenderTarget2D>> UndistortionDisplacementMapHolders;
 
 	/** Texture used to store temporary distortion displacement map when using map blending */
 	UPROPERTY(Transient)
-	TArray<UTextureRenderTarget2D*> DistortionDisplacementMapHolders;
+	TArray<TObjectPtr<UTextureRenderTarget2D>> DistortionDisplacementMapHolders;
 
 	/** The number of intermediate displacement maps needed to do map blending */
 	static constexpr int32 DisplacementMapHolderCount = 4;
@@ -322,7 +322,7 @@ public:
 
 	/** LensFile asset to use if DefaultLensFile is not desired */
 	UPROPERTY(EditAnywhere, BlueprintReadWrite, Category = "Lens File", Meta = (EditCondition = "!bUseDefaultLensFile"))
-	ULensFile* LensFile = nullptr;
+	TObjectPtr<ULensFile> LensFile = nullptr;
 };
 
 /** Structure that caches the inputs (and other useful bits) used when evaluating the Lens File */
diff --git a/Engine/Plugins/VirtualProduction/CameraCalibrationCore/Source/CameraCalibrationCore/Public/Tables/STMapTable.h b/Engine/Plugins/VirtualProduction/CameraCalibrationCore/Source/CameraCalibrationCore/Public/Tables/STMapTable.h
--- a/Engine/Plugins/VirtualProduction/CameraCalibrationCore/Source/CameraCalibrationCore/Public/Tables/STMapTable.h
+++ b/Engine/Plugins/VirtualProduction/CameraCalibrationCore/Source/CameraCalibrationCore/Public/Tables/STMapTable.h
@@ -24,11 +24,11 @@ struct CAMERACALIBRATIONCORE_API FDerivedDistortionData
 
 	/** Computed displacement map based on undistortion data */
 	UPROPERTY(Transient, VisibleAnywhere, Category = "Distortion")
-	UTextureRenderTarget2D* UndistortionDisplacementMap = nullptr;
+	TObjectPtr<UTextureRenderTarget2D> UndistortionDisplacementMap = nullptr;
 
 	/** Computed displacement map based on distortion data */
 	UPROPERTY(Transient, VisibleAnywhere, Category = "Distortion")
-	UTextureRenderTarget2D* DistortionDisplacementMap = nullptr;
+	TObjectPtr<UTextureRenderTarget2D> DistortionDisplacementMap = nullptr;
 
 	/** When dirty, derived data needs to be recomputed */
 	bool bIsDirty = true;
diff --git a/Engine/Plugins/VirtualProduction/DMX/DMXEngine/Source/DMXEditor/Private/Factories/DMXGDTFFactory.h b/Engine/Plugins/VirtualProduction/DMX/DMXEngine/Source/DMXEditor/Private/Factories/DMXGDTFFactory.h
--- a/Engine/Plugins/VirtualProduction/DMX/DMXEngine/Source/DMXEditor/Private/Factories/DMXGDTFFactory.h
+++ b/Engine/Plugins/VirtualProduction/DMX/DMXEngine/Source/DMXEditor/Private/Factories/DMXGDTFFactory.h
@@ -54,5 +54,5 @@ private:
 	bool bOperationCanceled;
 
 	UPROPERTY(transient)
-	UDMXGDTFImportUI* ImportUI;
+	TObjectPtr<UDMXGDTFImportUI> ImportUI;
 };
diff --git a/Engine/Plugins/VirtualProduction/DMX/DMXEngine/Source/DMXEditor/Private/Sequencer/TakeRecorderDMXLibrarySource.h b/Engine/Plugins/VirtualProduction/DMX/DMXEngine/Source/DMXEditor/Private/Sequencer/TakeRecorderDMXLibrarySource.h
--- a/Engine/Plugins/VirtualProduction/DMX/DMXEngine/Source/DMXEditor/Private/Sequencer/TakeRecorderDMXLibrarySource.h
+++ b/Engine/Plugins/VirtualProduction/DMX/DMXEngine/Source/DMXEditor/Private/Sequencer/TakeRecorderDMXLibrarySource.h
@@ -39,7 +39,7 @@ public:
 
 	/** DMX Library to record Patches' Fixture Functions from */
 	UPROPERTY(EditAnywhere, BlueprintReadWrite, Category = "Source", meta = (DisplayName = "DMX Library"))
-	UDMXLibrary* DMXLibrary;
+	TObjectPtr<UDMXLibrary> DMXLibrary;
 
 	/**
 	 * Dummy property to be replaced with the "Add all patches" button.
@@ -99,7 +99,7 @@ private:
 
 	/** Track recorder used by this source */
 	UPROPERTY()
-	UMovieSceneDMXLibraryTrackRecorder* TrackRecorder;
+	TObjectPtr<UMovieSceneDMXLibraryTrackRecorder> TrackRecorder;
 
 	/**
 	 * Stores an existing DMX Library track in the Sequence to be recorded or
diff --git a/Engine/Plugins/VirtualProduction/DMX/DMXEngine/Source/DMXRuntime/Public/DMXSubsystem.h b/Engine/Plugins/VirtualProduction/DMX/DMXEngine/Source/DMXRuntime/Public/DMXSubsystem.h
--- a/Engine/Plugins/VirtualProduction/DMX/DMXEngine/Source/DMXRuntime/Public/DMXSubsystem.h
+++ b/Engine/Plugins/VirtualProduction/DMX/DMXEngine/Source/DMXRuntime/Public/DMXSubsystem.h
@@ -316,5 +316,5 @@ private:
 private:
 	/** Strongly references all libraries at all times */
 	UPROPERTY(Transient)
-	TArray<UDMXLibrary*> LoadedDMXLibraries;
+	TArray<TObjectPtr<UDMXLibrary>> LoadedDMXLibraries;
 };
diff --git a/Engine/Plugins/VirtualProduction/DMX/DMXEngine/Source/DMXRuntime/Public/Library/DMXEntityFixturePatch.h b/Engine/Plugins/VirtualProduction/DMX/DMXEngine/Source/DMXRuntime/Public/Library/DMXEntityFixturePatch.h
--- a/Engine/Plugins/VirtualProduction/DMX/DMXEngine/Source/DMXRuntime/Public/Library/DMXEntityFixturePatch.h
+++ b/Engine/Plugins/VirtualProduction/DMX/DMXEngine/Source/DMXRuntime/Public/Library/DMXEntityFixturePatch.h
@@ -262,7 +262,7 @@ protected:
 
 	/** Property to point to the template parent fixture for details panel purposes */
 	UPROPERTY(VisibleDefaultsOnly, BlueprintReadWrite, BlueprintSetter = SetFixtureType, Category = "Fixture Patch", meta = (DisplayName = "Fixture Type"))
-	UDMXEntityFixtureType* ParentFixtureTypeTemplate;
+	TObjectPtr<UDMXEntityFixtureType> ParentFixtureTypeTemplate;
 
 	/** The Index of the Mode in the Fixture Type the Patch uses */
 	UPROPERTY(EditAnywhere, BlueprintReadOnly, Category = "Fixture Patch")
diff --git a/Engine/Plugins/VirtualProduction/DMX/DMXEngine/Source/DMXRuntime/Public/Library/DMXEntityFixtureType.h b/Engine/Plugins/VirtualProduction/DMX/DMXEngine/Source/DMXRuntime/Public/Library/DMXEntityFixtureType.h
--- a/Engine/Plugins/VirtualProduction/DMX/DMXEngine/Source/DMXRuntime/Public/Library/DMXEntityFixtureType.h
+++ b/Engine/Plugins/VirtualProduction/DMX/DMXEngine/Source/DMXRuntime/Public/Library/DMXEntityFixtureType.h
@@ -264,7 +264,7 @@ struct DMXRUNTIME_API FDMXEntityFixtureTypeConstructionParams
 
 	/** The DMX Library in which the Fixture Type will be constructed */
 	UPROPERTY(EditAnywhere, BlueprintReadWrite, Category = "Fixture Type")
-	UDMXLibrary* ParentDMXLibrary = nullptr;
+	TObjectPtr<UDMXLibrary> ParentDMXLibrary = nullptr;
 
 	/** The Category of the Fixture, useful for Filtering */
 	UPROPERTY(EditAnywhere, BlueprintReadWrite, Category = "Fixture Type")
@@ -322,7 +322,7 @@ public:
 
 	/** The GDTF file from which the Fixture Type was setup */
 	UPROPERTY(EditAnywhere, BlueprintReadOnly, Category = "Fixture Settings")
-	UDMXImport* DMXImport;
+	TObjectPtr<UDMXImport> DMXImport;
 
 	/** The Category of the Fixture, useful for Filtering */
 	UPROPERTY(EditAnywhere, BlueprintReadOnly, Category = "Fixture Settings", meta = (DisplayName = "DMX Category"))
@@ -336,7 +336,7 @@ public:
 	 * NOTE: Modulators only affect the patch's normalized values! Untouched values are still available when accesing raw values. 
 	 */
 	UPROPERTY(EditAnywhere, Instanced, Category = "Mode Settings", meta = (DisplayPriority = "50"))
-	TArray<UDMXModulator*> InputModulators;
+	TArray<TObjectPtr<UDMXModulator>> InputModulators;
 
 private:
 	/** Delegate that should be broadcast whenever a fixture type changed */
diff --git a/Engine/Plugins/VirtualProduction/DMX/DMXEngine/Source/DMXRuntime/Public/Library/DMXEntityReference.h b/Engine/Plugins/VirtualProduction/DMX/DMXEngine/Source/DMXRuntime/Public/Library/DMXEntityReference.h
--- a/Engine/Plugins/VirtualProduction/DMX/DMXEngine/Source/DMXRuntime/Public/Library/DMXEntityReference.h
+++ b/Engine/Plugins/VirtualProduction/DMX/DMXEngine/Source/DMXRuntime/Public/Library/DMXEntityReference.h
@@ -29,7 +29,7 @@ public:
 	 * Automatically set when calling SetEntity with a valid Entity.
 	 */
 	UPROPERTY(EditAnywhere, Category = "DMX")
-	UDMXLibrary* DMXLibrary;
+	TObjectPtr<UDMXLibrary> DMXLibrary;
 
 	/** Display the DMX Library asset picker. True by default, for Blueprint variables */
 	UPROPERTY() 
diff --git a/Engine/Plugins/VirtualProduction/DMX/DMXEngine/Source/DMXRuntime/Public/Library/DMXImport.h b/Engine/Plugins/VirtualProduction/DMX/DMXEngine/Source/DMXRuntime/Public/Library/DMXImport.h
--- a/Engine/Plugins/VirtualProduction/DMX/DMXEngine/Source/DMXRuntime/Public/Library/DMXImport.h
+++ b/Engine/Plugins/VirtualProduction/DMX/DMXEngine/Source/DMXRuntime/Public/Library/DMXImport.h
@@ -113,26 +113,26 @@ public:
 
 public:
 	UPROPERTY(VisibleAnywhere, Category = "Fixture Type", meta = (EditInline))
-	UDMXImportFixtureType* FixtureType;
+	TObjectPtr<UDMXImportFixtureType> FixtureType;
 
     UPROPERTY(VisibleAnywhere, Category = "Attribute Definitions", meta = (EditInline))
-    UDMXImportAttributeDefinitions* AttributeDefinitions;
+    TObjectPtr<UDMXImportAttributeDefinitions> AttributeDefinitions;
 
     UPROPERTY(VisibleAnywhere, Category = "Wheels", meta = (EditInline))
-    UDMXImportWheels* Wheels;
+    TObjectPtr<UDMXImportWheels> Wheels;
 
     UPROPERTY(VisibleAnywhere, Category = "Physical Descriptions", meta = (EditInline))
-    UDMXImportPhysicalDescriptions* PhysicalDescriptions;
+    TObjectPtr<UDMXImportPhysicalDescriptions> PhysicalDescriptions;
 
     UPROPERTY(VisibleAnywhere, Category = "Models", meta = (EditInline))
-    UDMXImportModels* Models;
+    TObjectPtr<UDMXImportModels> Models;
 
     UPROPERTY(VisibleAnywhere, Category = "Geometries", meta = (EditInline))
-    UDMXImportGeometries* Geometries;
+    TObjectPtr<UDMXImportGeometries> Geometries;
 
     UPROPERTY(VisibleAnywhere, Category = "DMXModes", meta = (EditInline))
-    UDMXImportDMXModes* DMXModes;
+    TObjectPtr<UDMXImportDMXModes> DMXModes;
 
     UPROPERTY(VisibleAnywhere, Category = "Protocols", meta = (EditInline))
-    UDMXImportProtocols* Protocols;
+    TObjectPtr<UDMXImportProtocols> Protocols;
 };
diff --git a/Engine/Plugins/VirtualProduction/DMX/DMXEngine/Source/DMXRuntime/Public/Library/DMXImportGDTF.h b/Engine/Plugins/VirtualProduction/DMX/DMXEngine/Source/DMXRuntime/Public/Library/DMXImportGDTF.h
--- a/Engine/Plugins/VirtualProduction/DMX/DMXEngine/Source/DMXRuntime/Public/Library/DMXImportGDTF.h
+++ b/Engine/Plugins/VirtualProduction/DMX/DMXEngine/Source/DMXRuntime/Public/Library/DMXImportGDTF.h
@@ -207,7 +207,7 @@ struct FDMXImportGDTFWheelSlot
     FDMXImportGDTFFilter Filter;
 
     UPROPERTY(VisibleAnywhere, Category = "Fixture Type")
-    UTexture2D* MediaFileName = nullptr;
+    TObjectPtr<UTexture2D> MediaFileName = nullptr;
 };
 
 USTRUCT(BlueprintType)
@@ -713,7 +713,7 @@ public:
     FString FixtureTypeID;
 
     UPROPERTY(VisibleAnywhere, Category = "Fixture Type")
-    UTexture2D* Thumbnail = nullptr;
+    TObjectPtr<UTexture2D> Thumbnail = nullptr;
 
     UPROPERTY(EditAnywhere, BlueprintReadOnly, Category = "Fixture Type")
     FString RefFT;
@@ -857,6 +857,6 @@ private:
 #if WITH_EDITORONLY_DATA
 	/** The Asset Import Data used to generate the GDTF asset or nullptr, if not generated from a GDTF file */
 	UPROPERTY()
-	UDMXGDTFAssetImportData* GDTFAssetImportData;
+	TObjectPtr<UDMXGDTFAssetImportData> GDTFAssetImportData;
 #endif
 };
diff --git a/Engine/Plugins/VirtualProduction/DMX/DMXEngine/Source/DMXRuntime/Public/Library/DMXLibrary.h b/Engine/Plugins/VirtualProduction/DMX/DMXEngine/Source/DMXRuntime/Public/Library/DMXLibrary.h
--- a/Engine/Plugins/VirtualProduction/DMX/DMXEngine/Source/DMXRuntime/Public/Library/DMXLibrary.h
+++ b/Engine/Plugins/VirtualProduction/DMX/DMXEngine/Source/DMXRuntime/Public/Library/DMXLibrary.h
@@ -252,11 +252,11 @@ protected:
 private:
 	/** All Fixture Types and Fixture Patches in the Library */
 	UPROPERTY()
-	TArray<UDMXEntity*> Entities;
+	TArray<TObjectPtr<UDMXEntity>> Entities;
 
 	/** The General Scene Description of this Library */
 	UPROPERTY()
-	UDMXMVRGeneralSceneDescription* GeneralSceneDescription;
+	TObjectPtr<UDMXMVRGeneralSceneDescription> GeneralSceneDescription;
 
 	/** The entity that was added last to the Library */
 	TWeakObjectPtr<UDMXEntity> LastAddedEntity;
diff --git a/Engine/Plugins/VirtualProduction/DMX/DMXEngine/Source/DMXRuntime/Public/MVR/DMXMVRAssetImportData.h b/Engine/Plugins/VirtualProduction/DMX/DMXEngine/Source/DMXRuntime/Public/MVR/DMXMVRAssetImportData.h
--- a/Engine/Plugins/VirtualProduction/DMX/DMXEngine/Source/DMXRuntime/Public/MVR/DMXMVRAssetImportData.h
+++ b/Engine/Plugins/VirtualProduction/DMX/DMXEngine/Source/DMXRuntime/Public/MVR/DMXMVRAssetImportData.h
@@ -35,7 +35,7 @@ private:
 #if WITH_EDITORONLY_DATA
 	/** The Asset Import Data used to generate the GDTF asset or nullptr, if not generated from a GDTF file */
 	UPROPERTY()
-	UAssetImportData* AssetImportData;
+	TObjectPtr<UAssetImportData> AssetImportData;
 
 	/** The raw source file as byte array, as it was imported */
 	UPROPERTY()
diff --git a/Engine/Plugins/VirtualProduction/DMX/DMXEngine/Source/DMXRuntime/Public/MVR/DMXMVRGeneralSceneDescription.h b/Engine/Plugins/VirtualProduction/DMX/DMXEngine/Source/DMXRuntime/Public/MVR/DMXMVRGeneralSceneDescription.h
--- a/Engine/Plugins/VirtualProduction/DMX/DMXEngine/Source/DMXRuntime/Public/MVR/DMXMVRGeneralSceneDescription.h
+++ b/Engine/Plugins/VirtualProduction/DMX/DMXEngine/Source/DMXRuntime/Public/MVR/DMXMVRGeneralSceneDescription.h
@@ -68,11 +68,11 @@ private:
 
 	/** The Root Node of the General Scene Description */
 	UPROPERTY()
-	UDMXMVRRootNode* RootNode;
+	TObjectPtr<UDMXMVRRootNode> RootNode;
 
 #if WITH_EDITORONLY_DATA
 	/** Import Data for this asset */
 	UPROPERTY()
-	UDMXMVRAssetImportData* MVRAssetImportData;
+	TObjectPtr<UDMXMVRAssetImportData> MVRAssetImportData;
 #endif
 };
diff --git a/Engine/Plugins/VirtualProduction/DMX/DMXEngine/Source/DMXRuntime/Public/MVR/DMXMVRSceneActor.h b/Engine/Plugins/VirtualProduction/DMX/DMXEngine/Source/DMXRuntime/Public/MVR/DMXMVRSceneActor.h
--- a/Engine/Plugins/VirtualProduction/DMX/DMXEngine/Source/DMXRuntime/Public/MVR/DMXMVRSceneActor.h
+++ b/Engine/Plugins/VirtualProduction/DMX/DMXEngine/Source/DMXRuntime/Public/MVR/DMXMVRSceneActor.h
@@ -113,7 +113,7 @@ private:
 
 	/** The root component to which all actors are attached initially */
 	UPROPERTY(VisibleAnywhere, Category = "Actor", AdvancedDisplay, Meta = (AllowPrivateAccess = true))
-	USceneComponent* MVRSceneRoot;
+	TObjectPtr<USceneComponent> MVRSceneRoot;
 
 private:
 	/** Gets the Fixture Patch of an Actor, or nullptr if the Fixture Patch cannot be retrieved anymore */
diff --git a/Engine/Plugins/VirtualProduction/DMX/DMXEngine/Source/DMXRuntime/Public/MVR/Types/DMXMVRChildListNode.h b/Engine/Plugins/VirtualProduction/DMX/DMXEngine/Source/DMXRuntime/Public/MVR/Types/DMXMVRChildListNode.h
--- a/Engine/Plugins/VirtualProduction/DMX/DMXEngine/Source/DMXRuntime/Public/MVR/Types/DMXMVRChildListNode.h
+++ b/Engine/Plugins/VirtualProduction/DMX/DMXEngine/Source/DMXRuntime/Public/MVR/Types/DMXMVRChildListNode.h
@@ -54,5 +54,5 @@ public:
 private:
 	/** A list of geometrical representation objects that are part of the object. */
 	UPROPERTY()
-	TArray<UDMXMVRParametricObjectNodeBase*> ParametricObjectNodes;
+	TArray<TObjectPtr<UDMXMVRParametricObjectNodeBase>> ParametricObjectNodes;
 };
diff --git a/Engine/Plugins/VirtualProduction/DMX/DMXEngine/Source/DMXRuntime/Public/MVR/Types/DMXMVRGroupObjectNode.h b/Engine/Plugins/VirtualProduction/DMX/DMXEngine/Source/DMXRuntime/Public/MVR/Types/DMXMVRGroupObjectNode.h
--- a/Engine/Plugins/VirtualProduction/DMX/DMXEngine/Source/DMXRuntime/Public/MVR/Types/DMXMVRGroupObjectNode.h
+++ b/Engine/Plugins/VirtualProduction/DMX/DMXEngine/Source/DMXRuntime/Public/MVR/Types/DMXMVRGroupObjectNode.h
@@ -59,5 +59,5 @@ public:
 private:
 	/** A list of graphic objects that are part of the group. */
 	UPROPERTY()
-	UDMXMVRChildListNode* ChildListNode;
+	TObjectPtr<UDMXMVRChildListNode> ChildListNode;
 };
diff --git a/Engine/Plugins/VirtualProduction/DMX/DMXEngine/Source/DMXRuntime/Public/MVR/Types/DMXMVRLayerNode.h b/Engine/Plugins/VirtualProduction/DMX/DMXEngine/Source/DMXRuntime/Public/MVR/Types/DMXMVRLayerNode.h
--- a/Engine/Plugins/VirtualProduction/DMX/DMXEngine/Source/DMXRuntime/Public/MVR/Types/DMXMVRLayerNode.h
+++ b/Engine/Plugins/VirtualProduction/DMX/DMXEngine/Source/DMXRuntime/Public/MVR/Types/DMXMVRLayerNode.h
@@ -65,5 +65,5 @@ public:
 private:
 	/** A list of graphic objects that are part of the group. */
 	UPROPERTY()
-	UDMXMVRChildListNode* ChildListNode;
+	TObjectPtr<UDMXMVRChildListNode> ChildListNode;
 };
diff --git a/Engine/Plugins/VirtualProduction/DMX/DMXEngine/Source/DMXRuntime/Public/MVR/Types/DMXMVRLayersNode.h b/Engine/Plugins/VirtualProduction/DMX/DMXEngine/Source/DMXRuntime/Public/MVR/Types/DMXMVRLayersNode.h
--- a/Engine/Plugins/VirtualProduction/DMX/DMXEngine/Source/DMXRuntime/Public/MVR/Types/DMXMVRLayersNode.h
+++ b/Engine/Plugins/VirtualProduction/DMX/DMXEngine/Source/DMXRuntime/Public/MVR/Types/DMXMVRLayersNode.h
@@ -46,5 +46,5 @@ public:
 private:
 	/** A[n array of] layer representation[s]. */
 	UPROPERTY()
-	TArray<UDMXMVRLayerNode*> LayerNodes;
+	TArray<TObjectPtr<UDMXMVRLayerNode>> LayerNodes;
 };
diff --git a/Engine/Plugins/VirtualProduction/DMX/DMXEngine/Source/DMXRuntime/Public/MVR/Types/DMXMVRRootNode.h b/Engine/Plugins/VirtualProduction/DMX/DMXEngine/Source/DMXRuntime/Public/MVR/Types/DMXMVRRootNode.h
--- a/Engine/Plugins/VirtualProduction/DMX/DMXEngine/Source/DMXRuntime/Public/MVR/Types/DMXMVRRootNode.h
+++ b/Engine/Plugins/VirtualProduction/DMX/DMXEngine/Source/DMXRuntime/Public/MVR/Types/DMXMVRRootNode.h
@@ -59,9 +59,9 @@ public:
 private:
 	/** This node contains a collection of user data nodes defined and used by provider applications if required. */
 	UPROPERTY()
-	UDMXMVRUserDataNode* UserDataNode;
+	TObjectPtr<UDMXMVRUserDataNode> UserDataNode;
 
 	/** This node contains information about the scene. */
 	UPROPERTY()
-	UDMXMVRSceneNode* SceneNode;
+	TObjectPtr<UDMXMVRSceneNode> SceneNode;
 };
diff --git a/Engine/Plugins/VirtualProduction/DMX/DMXEngine/Source/DMXRuntime/Public/MVR/Types/DMXMVRSceneNode.h b/Engine/Plugins/VirtualProduction/DMX/DMXEngine/Source/DMXRuntime/Public/MVR/Types/DMXMVRSceneNode.h
--- a/Engine/Plugins/VirtualProduction/DMX/DMXEngine/Source/DMXRuntime/Public/MVR/Types/DMXMVRSceneNode.h
+++ b/Engine/Plugins/VirtualProduction/DMX/DMXEngine/Source/DMXRuntime/Public/MVR/Types/DMXMVRSceneNode.h
@@ -38,5 +38,5 @@ public:
 private:
 	/** This node defines a list of layers inside the scene. The layer is a container of graphical objects defining a local coordinate system. */
 	UPROPERTY()
-	UDMXMVRLayersNode* LayersNode;
+	TObjectPtr<UDMXMVRLayersNode> LayersNode;
 };
diff --git a/Engine/Plugins/VirtualProduction/DMX/DMXEngine/Source/DMXRuntime/Public/MVR/Types/DMXMVRUserDataNode.h b/Engine/Plugins/VirtualProduction/DMX/DMXEngine/Source/DMXRuntime/Public/MVR/Types/DMXMVRUserDataNode.h
--- a/Engine/Plugins/VirtualProduction/DMX/DMXEngine/Source/DMXRuntime/Public/MVR/Types/DMXMVRUserDataNode.h
+++ b/Engine/Plugins/VirtualProduction/DMX/DMXEngine/Source/DMXRuntime/Public/MVR/Types/DMXMVRUserDataNode.h
@@ -29,5 +29,5 @@ public:
 private:
 	/** UE Specific: The UE's proprietary data node */
 	UPROPERTY()
-	UDMXMVRUnrealEngineDataNode* UnrealEngineDataNode;
+	TObjectPtr<UDMXMVRUnrealEngineDataNode> UnrealEngineDataNode;
 };
diff --git a/Engine/Plugins/VirtualProduction/DMX/DMXEngine/Source/DMXRuntime/Public/Sequencer/MovieSceneDMXLibrarySection.h b/Engine/Plugins/VirtualProduction/DMX/DMXEngine/Source/DMXRuntime/Public/Sequencer/MovieSceneDMXLibrarySection.h
--- a/Engine/Plugins/VirtualProduction/DMX/DMXEngine/Source/DMXRuntime/Public/Sequencer/MovieSceneDMXLibrarySection.h
+++ b/Engine/Plugins/VirtualProduction/DMX/DMXEngine/Source/DMXRuntime/Public/Sequencer/MovieSceneDMXLibrarySection.h
@@ -65,7 +65,7 @@ struct FDMXFixturePatchChannel
 
 	/** The outer library of the channel */
 	UPROPERTY()
-	UDMXLibrary* DMXLibrary;
+	TObjectPtr<UDMXLibrary> DMXLibrary;
 
 	/** Points to the Fixture Patch */
 	UPROPERTY()
diff --git a/Engine/Plugins/VirtualProduction/DMX/DMXEngine/Source/DMXRuntime/Public/Sequencer/MovieSceneDMXLibraryTemplate.h b/Engine/Plugins/VirtualProduction/DMX/DMXEngine/Source/DMXRuntime/Public/Sequencer/MovieSceneDMXLibraryTemplate.h
--- a/Engine/Plugins/VirtualProduction/DMX/DMXEngine/Source/DMXRuntime/Public/Sequencer/MovieSceneDMXLibraryTemplate.h
+++ b/Engine/Plugins/VirtualProduction/DMX/DMXEngine/Source/DMXRuntime/Public/Sequencer/MovieSceneDMXLibraryTemplate.h
@@ -37,5 +37,5 @@ private:
 
 private:
 	UPROPERTY()
-	const UMovieSceneDMXLibrarySection* Section;
+	TObjectPtr<const UMovieSceneDMXLibrarySection> Section;
 };
diff --git a/Engine/Plugins/VirtualProduction/DMX/DMXEngine/Source/DMXRuntime/Public/Sequencer/MovieSceneDMXLibraryTrack.h b/Engine/Plugins/VirtualProduction/DMX/DMXEngine/Source/DMXRuntime/Public/Sequencer/MovieSceneDMXLibraryTrack.h
--- a/Engine/Plugins/VirtualProduction/DMX/DMXEngine/Source/DMXRuntime/Public/Sequencer/MovieSceneDMXLibraryTrack.h
+++ b/Engine/Plugins/VirtualProduction/DMX/DMXEngine/Source/DMXRuntime/Public/Sequencer/MovieSceneDMXLibraryTrack.h
@@ -43,9 +43,9 @@ public:
 private:
 	/** The sections owned by this track .*/
 	UPROPERTY()
-	TArray<UMovieSceneSection*> Sections;
+	TArray<TObjectPtr<UMovieSceneSection>> Sections;
 
 	/** The DMX Library to manipulate */
 	UPROPERTY()
-	UDMXLibrary* Library;
+	TObjectPtr<UDMXLibrary> Library;
 };
diff --git a/Engine/Plugins/VirtualProduction/DMX/DMXFixtures/Source/DMXFixtures/Public/DMXFixtureActor.h b/Engine/Plugins/VirtualProduction/DMX/DMXFixtures/Source/DMXFixtures/Public/DMXFixtureActor.h
--- a/Engine/Plugins/VirtualProduction/DMX/DMXFixtures/Source/DMXFixtures/Public/DMXFixtureActor.h
+++ b/Engine/Plugins/VirtualProduction/DMX/DMXFixtures/Source/DMXFixtures/Public/DMXFixtureActor.h
@@ -73,13 +73,13 @@ public:
 
 	// HIERARCHY---------------------------------
 	UPROPERTY(EditDefaultsOnly, BlueprintReadOnly, Category = "DMX Light Fixture")
-	USceneComponent* Base;
+	TObjectPtr<USceneComponent> Base;
 
 	UPROPERTY(EditDefaultsOnly, BlueprintReadOnly, Category = "DMX Light Fixture")
-	USceneComponent* Yoke;
+	TObjectPtr<USceneComponent> Yoke;
 
 	UPROPERTY(EditDefaultsOnly, BlueprintReadOnly, Category = "DMX Light Fixture")
-	USceneComponent* Head;
+	TObjectPtr<USceneComponent> Head;
 
 	// FUNCTIONS---------------------------------
 
@@ -153,47 +153,47 @@ public:
 	// DMX COMPONENT -----------------------------
 
 	UPROPERTY(VisibleAnywhere, BlueprintReadOnly, Category = "DMX Light Fixture")
-	class UDMXComponent* DMX;
+	TObjectPtr<class UDMXComponent> DMX;
 
 	// COMPONENTS ---------------------------------
 
 	UPROPERTY(BlueprintReadOnly, Category = "DMX Light Fixture")
-	TArray<UStaticMeshComponent*> StaticMeshComponents;
+	TArray<TObjectPtr<UStaticMeshComponent>> StaticMeshComponents;
 	
 	UPROPERTY(EditDefaultsOnly, BlueprintReadOnly, Category = "DMX Light Fixture")
-	class USpotLightComponent* SpotLight;
+	TObjectPtr<class USpotLightComponent> SpotLight;
 
 	UPROPERTY(EditDefaultsOnly, BlueprintReadOnly, Category = "DMX Light Fixture")
-	class UPointLightComponent* PointLight;
+	TObjectPtr<class UPointLightComponent> PointLight;
 
 	UPROPERTY(EditDefaultsOnly, BlueprintReadOnly, Category = "DMX Light Fixture")
-	class UArrowComponent* OcclusionDirection;
+	TObjectPtr<class UArrowComponent> OcclusionDirection;
 
 
 	// MATERIALS ---------------------------------
 
 	UPROPERTY(EditAnywhere, AdvancedDisplay, Category = "DMX Light Fixture")
-	UMaterialInstance* LensMaterialInstance;
+	TObjectPtr<UMaterialInstance> LensMaterialInstance;
 
 	UPROPERTY(EditAnywhere, AdvancedDisplay, Category = "DMX Light Fixture")
-	UMaterialInstance* BeamMaterialInstance;
+	TObjectPtr<UMaterialInstance> BeamMaterialInstance;
 
 	UPROPERTY(EditAnywhere, AdvancedDisplay, Category = "DMX Light Fixture")
-	UMaterialInstance* SpotLightMaterialInstance;
+	TObjectPtr<UMaterialInstance> SpotLightMaterialInstance;
 
 	UPROPERTY(EditAnywhere, AdvancedDisplay, Category = "DMX Light Fixture")
-	UMaterialInstance* PointLightMaterialInstance;
+	TObjectPtr<UMaterialInstance> PointLightMaterialInstance;
 
 	UPROPERTY(BlueprintReadOnly, Category = "DMX Light Fixture")
-	UMaterialInstanceDynamic* DynamicMaterialLens;
+	TObjectPtr<UMaterialInstanceDynamic> DynamicMaterialLens;
 
 	UPROPERTY(BlueprintReadOnly, Category = "DMX Light Fixture")
-	UMaterialInstanceDynamic* DynamicMaterialBeam;
+	TObjectPtr<UMaterialInstanceDynamic> DynamicMaterialBeam;
 
 	UPROPERTY(BlueprintReadOnly, Category = "DMX Light Fixture")
-	UMaterialInstanceDynamic* DynamicMaterialSpotLight;
+	TObjectPtr<UMaterialInstanceDynamic> DynamicMaterialSpotLight;
 
 	UPROPERTY(BlueprintReadOnly, Category = "DMX Light Fixture")
-	UMaterialInstanceDynamic* DynamicMaterialPointLight;
+	TObjectPtr<UMaterialInstanceDynamic> DynamicMaterialPointLight;
 
 };
diff --git a/Engine/Plugins/VirtualProduction/DMX/DMXFixtures/Source/DMXFixtures/Public/DMXFixtureActorMatrix.h b/Engine/Plugins/VirtualProduction/DMX/DMXFixtures/Source/DMXFixtures/Public/DMXFixtureActorMatrix.h
--- a/Engine/Plugins/VirtualProduction/DMX/DMXFixtures/Source/DMXFixtures/Public/DMXFixtureActorMatrix.h
+++ b/Engine/Plugins/VirtualProduction/DMX/DMXFixtures/Source/DMXFixtures/Public/DMXFixtureActorMatrix.h
@@ -38,7 +38,7 @@ public:
 	FUpdateTextureRegion2D* TextureRegion;
 
 	UPROPERTY(EditDefaultsOnly, BlueprintReadOnly, Category = "DMX Matrix Fixture")
-	UProceduralMeshComponent* MatrixHead;
+	TObjectPtr<UProceduralMeshComponent> MatrixHead;
 
 	UPROPERTY(EditAnywhere, BlueprintReadOnly, Category = "DMX Matrix Fixture")
 	float MatrixWidth;
diff --git a/Engine/Plugins/VirtualProduction/DMX/DMXPixelMapping/Source/DMXPixelMappingRuntime/Public/Components/DMXPixelMappingBaseComponent.h b/Engine/Plugins/VirtualProduction/DMX/DMXPixelMapping/Source/DMXPixelMappingRuntime/Public/Components/DMXPixelMappingBaseComponent.h
--- a/Engine/Plugins/VirtualProduction/DMX/DMXPixelMapping/Source/DMXPixelMappingRuntime/Public/Components/DMXPixelMappingBaseComponent.h
+++ b/Engine/Plugins/VirtualProduction/DMX/DMXPixelMapping/Source/DMXPixelMappingRuntime/Public/Components/DMXPixelMappingBaseComponent.h
@@ -235,7 +235,7 @@ public:
 
 	/** Array of children belong to this component */
 	UPROPERTY()
-	TArray<UDMXPixelMappingBaseComponent*> Children;
+	TArray<TObjectPtr<UDMXPixelMappingBaseComponent>> Children;
 
 	/** Returns the parent. May be nullptr when the the component is creating or destroying */
 	FORCEINLINE UDMXPixelMappingBaseComponent* GetParent() const { return WeakParent.Get(); }
@@ -250,7 +250,7 @@ public:
 
 	/** Parent component */
 	UPROPERTY(Meta = (DeprecatedProperty, DeprecationMessage = "Leads to entangled references. Use GetParent() or WeakParent instead."))
-	UDMXPixelMappingBaseComponent* Parent_DEPRECATED;
+	TObjectPtr<UDMXPixelMappingBaseComponent> Parent_DEPRECATED;
 
 protected:
 	/** Called when the component was added to a parent */
diff --git a/Engine/Plugins/VirtualProduction/DMX/DMXPixelMapping/Source/DMXPixelMappingRuntime/Public/Components/DMXPixelMappingFixtureGroupItemComponent.h b/Engine/Plugins/VirtualProduction/DMX/DMXPixelMapping/Source/DMXPixelMappingRuntime/Public/Components/DMXPixelMappingFixtureGroupItemComponent.h
--- a/Engine/Plugins/VirtualProduction/DMX/DMXPixelMapping/Source/DMXPixelMappingRuntime/Public/Components/DMXPixelMappingFixtureGroupItemComponent.h
+++ b/Engine/Plugins/VirtualProduction/DMX/DMXPixelMapping/Source/DMXPixelMappingRuntime/Public/Components/DMXPixelMappingFixtureGroupItemComponent.h
@@ -118,7 +118,7 @@ public:
 	TArray<TSubclassOf<UDMXModulator>> ModulatorClasses;
 
 	UPROPERTY()
-	TArray<UDMXModulator*> Modulators;
+	TArray<TObjectPtr<UDMXModulator>> Modulators;
 
 	/** Index of the cell pixel in downsample target buffer */
 	int32 DownsamplePixelIndex;
diff --git a/Engine/Plugins/VirtualProduction/DMX/DMXPixelMapping/Source/DMXPixelMappingRuntime/Public/Components/DMXPixelMappingMatrixComponent.h b/Engine/Plugins/VirtualProduction/DMX/DMXPixelMapping/Source/DMXPixelMappingRuntime/Public/Components/DMXPixelMappingMatrixComponent.h
--- a/Engine/Plugins/VirtualProduction/DMX/DMXPixelMapping/Source/DMXPixelMappingRuntime/Public/Components/DMXPixelMappingMatrixComponent.h
+++ b/Engine/Plugins/VirtualProduction/DMX/DMXPixelMapping/Source/DMXPixelMappingRuntime/Public/Components/DMXPixelMappingMatrixComponent.h
@@ -143,7 +143,7 @@ public:
 
 	/** The actual modulator instances */
 	UPROPERTY()
-	TArray<UDMXModulator*> Modulators;
+	TArray<TObjectPtr<UDMXModulator>> Modulators;
 
 	UPROPERTY()
 	FIntPoint CoordinateGrid;
diff --git a/Engine/Plugins/VirtualProduction/DMX/DMXPixelMapping/Source/DMXPixelMappingRuntime/Public/Components/DMXPixelMappingRendererComponent.h b/Engine/Plugins/VirtualProduction/DMX/DMXPixelMapping/Source/DMXPixelMappingRuntime/Public/Components/DMXPixelMappingRendererComponent.h
--- a/Engine/Plugins/VirtualProduction/DMX/DMXPixelMapping/Source/DMXPixelMappingRuntime/Public/Components/DMXPixelMappingRendererComponent.h
+++ b/Engine/Plugins/VirtualProduction/DMX/DMXPixelMapping/Source/DMXPixelMappingRuntime/Public/Components/DMXPixelMappingRendererComponent.h
@@ -170,11 +170,11 @@ public:
 
 	/** Texture to Downsampling */
 	UPROPERTY(EditAnywhere, BlueprintReadOnly, Category = "Render Settings")
-	UTexture* InputTexture;
+	TObjectPtr<UTexture> InputTexture;
 
 	/** Material to Downsampling */
 	UPROPERTY(EditAnywhere, BlueprintReadOnly, Category = "Render Settings", meta = (DisplayName = "User Interface Material"))
-	UMaterialInterface* InputMaterial;
+	TObjectPtr<UMaterialInterface> InputMaterial;
 
 	/** UMG to Downsampling */
 	UPROPERTY(EditAnywhere, BlueprintReadOnly, Category = "Render Settings")
@@ -196,19 +196,19 @@ private:
 #if WITH_EDITORONLY_DATA
 	/** Editor preview output target */
 	UPROPERTY(Transient)
-	UTextureRenderTarget2D* PreviewRenderTarget;
+	TObjectPtr<UTextureRenderTarget2D> PreviewRenderTarget;
 #endif // WITH_EDITORONLY_DATA
 
 	/** Material of UMG texture to downsample */
 	UPROPERTY(Transient)
-	UTextureRenderTarget2D* InputRenderTarget;
+	TObjectPtr<UTextureRenderTarget2D> InputRenderTarget;
 
 	/** Reference to renderer */
 	TSharedPtr<IDMXPixelMappingRenderer> PixelMappingRenderer;
 
 	/** UMG widget for downsampling */
 	UPROPERTY(Transient)
-	UUserWidget* UserWidget;
+	TObjectPtr<UUserWidget> UserWidget;
 
 #if WITH_EDITORONLY_DATA
 	/** Canvas for all UI downsamping component widgets */
@@ -226,7 +226,7 @@ private:
 
 	/** GPU downsample pixel buffer target texture */
 	UPROPERTY(Transient)
-	UTextureRenderTarget2D* DownsampleBufferTarget;
+	TObjectPtr<UTextureRenderTarget2D> DownsampleBufferTarget;
 
 	/** CPU downsample pixel buffer */
 	TArray<FLinearColor> DownsampleBuffer;
diff --git a/Engine/Plugins/VirtualProduction/DMX/DMXPixelMapping/Source/DMXPixelMappingRuntime/Public/DMXPixelMapping.h b/Engine/Plugins/VirtualProduction/DMX/DMXPixelMapping/Source/DMXPixelMappingRuntime/Public/DMXPixelMapping.h
--- a/Engine/Plugins/VirtualProduction/DMX/DMXPixelMapping/Source/DMXPixelMappingRuntime/Public/DMXPixelMapping.h
+++ b/Engine/Plugins/VirtualProduction/DMX/DMXPixelMapping/Source/DMXPixelMappingRuntime/Public/DMXPixelMapping.h
@@ -138,13 +138,13 @@ public:
 public:
 	/** Holds the reference to root component */
 	UPROPERTY()
-	UDMXPixelMappingRootComponent* RootComponent;
+	TObjectPtr<UDMXPixelMappingRootComponent> RootComponent;
 
 
 #if WITH_EDITORONLY_DATA
 	/** Holds the Thumbnail for asset */
 	UPROPERTY()
-	UTexture* ThumbnailImage;
+	TObjectPtr<UTexture> ThumbnailImage;
 
 	/** DEPRECATED 4.27. Use FDMXPixelMappingMatrixComponent::GetOnMatrixChanged() instead */
 	FOnEditorRebuildChildrenComponentsDelegate OnEditorRebuildChildrenComponentsDelegate_DEPRECATED;
diff --git a/Engine/Plugins/VirtualProduction/LevelSnapshots/Source/LevelShapshotFilters/Public/Builtin/ActorDependentPropertyFilter.h b/Engine/Plugins/VirtualProduction/LevelSnapshots/Source/LevelShapshotFilters/Public/Builtin/ActorDependentPropertyFilter.h
--- a/Engine/Plugins/VirtualProduction/LevelSnapshots/Source/LevelShapshotFilters/Public/Builtin/ActorDependentPropertyFilter.h
+++ b/Engine/Plugins/VirtualProduction/LevelSnapshots/Source/LevelShapshotFilters/Public/Builtin/ActorDependentPropertyFilter.h
@@ -41,15 +41,15 @@ public:
 	
     /* We run IsActorValid on this filter. IsPropertyValid uses one of the below filters depending on this filter. */
 	UPROPERTY(EditAnywhere, Instanced, Category = "Config")
-	ULevelSnapshotFilter* ActorFilter;
+	TObjectPtr<ULevelSnapshotFilter> ActorFilter;
 
     /* Used by IsPropertyValid when ActorFilter->IsActorValid returns Include */
 	UPROPERTY(EditAnywhere, Instanced, Category = "Config")
-	ULevelSnapshotFilter* IncludePropertyFilter;
+	TObjectPtr<ULevelSnapshotFilter> IncludePropertyFilter;
 
 	/* Used by IsPropertyValid when ActorFilter->IsActorValid returns Exclude */
 	UPROPERTY(EditAnywhere, Instanced, Category = "Config")
-	ULevelSnapshotFilter* ExcludePropertyFilter;
+	TObjectPtr<ULevelSnapshotFilter> ExcludePropertyFilter;
 
 	/* Determines what filter IsPropertyValid is supposed to use when IsActorValid returns DoNotCare. */
 	UPROPERTY(EditAnywhere, Category = "Config")
@@ -57,5 +57,5 @@ public:
 	
 	/* Used by IsPropertyValid when ActorFilter->IsActorValid returns DoNotCare and DoNotCareHandling == UseDoNotCareFilter. */
 	UPROPERTY(EditAnywhere, Instanced, Category = "Config")
-	ULevelSnapshotFilter* DoNotCarePropertyFilter;
+	TObjectPtr<ULevelSnapshotFilter> DoNotCarePropertyFilter;
 };
diff --git a/Engine/Plugins/VirtualProduction/LevelSnapshots/Source/LevelShapshotFilters/Public/Builtin/BlueprintOnly/NegationFilter.h b/Engine/Plugins/VirtualProduction/LevelSnapshots/Source/LevelShapshotFilters/Public/Builtin/BlueprintOnly/NegationFilter.h
--- a/Engine/Plugins/VirtualProduction/LevelSnapshots/Source/LevelShapshotFilters/Public/Builtin/BlueprintOnly/NegationFilter.h
+++ b/Engine/Plugins/VirtualProduction/LevelSnapshots/Source/LevelShapshotFilters/Public/Builtin/BlueprintOnly/NegationFilter.h
@@ -51,8 +51,8 @@ public:
 private:
 	
 	UPROPERTY(EditAnywhere, Category = "Level Snapshots")
-	ULevelSnapshotFilter* Child;
+	TObjectPtr<ULevelSnapshotFilter> Child;
 
 	UPROPERTY(EditAnywhere, Instanced, Category = "Level Snapshots")
-	ULevelSnapshotFilter* InstancedChild;
+	TObjectPtr<ULevelSnapshotFilter> InstancedChild;
 };
diff --git a/Engine/Plugins/VirtualProduction/LevelSnapshots/Source/LevelShapshotFilters/Public/Builtin/BlueprintOnly/ParentFilter.h b/Engine/Plugins/VirtualProduction/LevelSnapshots/Source/LevelShapshotFilters/Public/Builtin/BlueprintOnly/ParentFilter.h
--- a/Engine/Plugins/VirtualProduction/LevelSnapshots/Source/LevelShapshotFilters/Public/Builtin/BlueprintOnly/ParentFilter.h
+++ b/Engine/Plugins/VirtualProduction/LevelSnapshots/Source/LevelShapshotFilters/Public/Builtin/BlueprintOnly/ParentFilter.h
@@ -49,8 +49,8 @@ protected:
 private:
 
 	UPROPERTY(EditAnywhere, Category = "Level Snapshots")
-	TArray<ULevelSnapshotFilter*> Children;
+	TArray<TObjectPtr<ULevelSnapshotFilter>> Children;
 
 	UPROPERTY(EditAnywhere, Instanced, Category = "Level Snapshots")
-	TArray<ULevelSnapshotFilter*> InstancedChildren;
+	TArray<TObjectPtr<ULevelSnapshotFilter>> InstancedChildren;
 };
diff --git a/Engine/Plugins/VirtualProduction/LevelSnapshots/Source/LevelShapshotFilters/Public/LevelSnapshotFilterParams.h b/Engine/Plugins/VirtualProduction/LevelSnapshots/Source/LevelShapshotFilters/Public/LevelSnapshotFilterParams.h
--- a/Engine/Plugins/VirtualProduction/LevelSnapshots/Source/LevelShapshotFilters/Public/LevelSnapshotFilterParams.h
+++ b/Engine/Plugins/VirtualProduction/LevelSnapshots/Source/LevelShapshotFilters/Public/LevelSnapshotFilterParams.h
@@ -16,11 +16,11 @@ struct LEVELSNAPSHOTFILTERS_API FIsActorValidParams
 
 	/** The actor saved in the snapshot */
 	UPROPERTY(BlueprintReadWrite, Category = "Level Snapshots")
-	AActor* SnapshotActor = nullptr;
+	TObjectPtr<AActor> SnapshotActor = nullptr;
 	
 	/** The actor equivalent to LevelActor: it exists in the world */
 	UPROPERTY(BlueprintReadWrite, Category = "Level Snapshots")
-	AActor* LevelActor = nullptr;
+	TObjectPtr<AActor> LevelActor = nullptr;
 
 	FIsActorValidParams() = default;
 	FIsActorValidParams(AActor* InSnapshotActor, AActor* InLevelActor)
@@ -36,11 +36,11 @@ struct LEVELSNAPSHOTFILTERS_API FIsPropertyValidParams
 
 	/** The actor saved in the snapshot */
 	UPROPERTY(BlueprintReadOnly, Category = "Level Snapshots")
-	AActor* SnapshotActor = nullptr;
+	TObjectPtr<AActor> SnapshotActor = nullptr;
 	
 	/** The actor equivalent to LevelActor: it exists in the world */
 	UPROPERTY(BlueprintReadOnly, Category = "Level Snapshots")
-	AActor* LevelActor = nullptr;
+	TObjectPtr<AActor> LevelActor = nullptr;
 
 	/** For passing to FProperty::ContainerPtrToValuePtr. This is either SnapshotActor or a subobject thereof. */
 	UPROPERTY(BlueprintReadOnly, Category = "Level Snapshots")
@@ -106,7 +106,7 @@ struct LEVELSNAPSHOTFILTERS_API FIsAddedActorValidParams
 
 	/** This actor was added to the level since the snapshot was taken. */
 	UPROPERTY(BlueprintReadOnly, Category = "Level Snapshots")
-	AActor* NewActor = nullptr;
+	TObjectPtr<AActor> NewActor = nullptr;
 
 	FIsAddedActorValidParams() = default;
 	explicit FIsAddedActorValidParams(AActor* NewActor)
@@ -121,11 +121,11 @@ struct LEVELSNAPSHOTFILTERS_API FIsDeletedComponentValidParams
 
 	/** This component was removed from the actor. This instance exists in a transient snapshot world; it does not exist in the editor world. */
 	UPROPERTY(BlueprintReadOnly, Category = "Level Snapshots")
-	UActorComponent* DeletedComponent = nullptr;
+	TObjectPtr<UActorComponent> DeletedComponent = nullptr;
 
 	/** The actor the component was removed from; This instance exists in the editor world. */
 	UPROPERTY(BlueprintReadOnly, Category = "Level Snapshots")
-	AActor* EditorActor = nullptr;
+	TObjectPtr<AActor> EditorActor = nullptr;
 
 	FIsDeletedComponentValidParams() = default;
 	explicit FIsDeletedComponentValidParams(UActorComponent* DeletedComponent, AActor* EditorActor)
@@ -141,7 +141,7 @@ struct LEVELSNAPSHOTFILTERS_API FIsAddedComponentValidParams
 
 	/** This component was added to the actor. This is an instance in the editor world. */
 	UPROPERTY(BlueprintReadOnly, Category = "Level Snapshots")
-	UActorComponent* AddedComponent = nullptr;
+	TObjectPtr<UActorComponent> AddedComponent = nullptr;
 
 	FIsAddedComponentValidParams() = default;
 	explicit FIsAddedComponentValidParams(UActorComponent* AddedComponent)
diff --git a/Engine/Plugins/VirtualProduction/LevelSnapshots/Source/LevelSnapshots/Private/Tests/Types/ActorWithReferencesInCDO.h b/Engine/Plugins/VirtualProduction/LevelSnapshots/Source/LevelSnapshots/Private/Tests/Types/ActorWithReferencesInCDO.h
--- a/Engine/Plugins/VirtualProduction/LevelSnapshots/Source/LevelSnapshots/Private/Tests/Types/ActorWithReferencesInCDO.h
+++ b/Engine/Plugins/VirtualProduction/LevelSnapshots/Source/LevelSnapshots/Private/Tests/Types/ActorWithReferencesInCDO.h
@@ -15,7 +15,7 @@ struct FExternalReferenceDummy
 	GENERATED_BODY()
 
 	UPROPERTY(EditAnywhere, Category = "Level Snapshots")
-	UObject* Object;
+	TObjectPtr<UObject> Object;
 
 	friend uint32 GetTypeHash(const FExternalReferenceDummy& Struct)
 	{
@@ -65,8 +65,8 @@ public:
 	/******************** External references  ********************/
 	
 	UPROPERTY(EditAnywhere, Category = "Level Snapshots")
-	UStaticMesh* CubeMesh = nullptr;
+	TObjectPtr<UStaticMesh> CubeMesh = nullptr;
 
 	UPROPERTY(EditAnywhere, Category = "Level Snapshots")
-	UStaticMesh* CylinderMesh = nullptr;
+	TObjectPtr<UStaticMesh> CylinderMesh = nullptr;
 };
