DMX - Minimal version of a Fixture Actor

diff --git a/Engine/Plugins/VirtualProduction/DMX/DMXFixtures/Source/DMXFixtures/Private/DMXFixtureActor.cpp b/Engine/Plugins/VirtualProduction/DMX/DMXFixtures/Source/DMXFixtures/Private/DMXFixtureActor.cpp
--- a/Engine/Plugins/VirtualProduction/DMX/DMXFixtures/Source/DMXFixtures/Private/DMXFixtureActor.cpp
+++ b/Engine/Plugins/VirtualProduction/DMX/DMXFixtures/Source/DMXFixtures/Private/DMXFixtureActor.cpp
@@ -3,6 +3,7 @@
 #include "DMXFixtureActor.h"
 
 #include "DMXStats.h"
+#include "Game/DMXComponent.h"
 
 #include "Components/StaticMeshComponent.h"
 
@@ -11,12 +12,6 @@ DECLARE_CYCLE_STAT(TEXT("FixtureActor Push Normalized Values"), STAT_FixtureActo
 
 ADMXFixtureActor::ADMXFixtureActor()
 {
-	PrimaryActorTick.bCanEverTick = true;
-
-#if WITH_EDITORONLY_DATA
-	bRunConstructionScriptOnDrag = 0;
-#endif
-
 	Base = CreateDefaultSubobject<USceneComponent>(TEXT("Base"));
 	RootComponent = Base;
 
@@ -39,8 +34,6 @@ ADMXFixtureActor::ADMXFixtureActor()
 	OcclusionDirection = CreateDefaultSubobject<UArrowComponent>(TEXT("Occlusion"));
 	OcclusionDirection->SetupAttachment(Head);
 
-	DMX = CreateDefaultSubobject<UDMXComponent>(TEXT("DMX"));
-
 	// set default values
 	LightIntensityMax = 2000;
 	LightDistanceMax = 1000;
@@ -193,41 +186,6 @@ void ADMXFixtureActor::FeedFixtureData()
 	PointLight->SetAttenuationRadius(LightDistanceMax);
 }
 
-void ADMXFixtureActor::InterpolateDMXComponents(float DeltaSeconds)
-{
-	// Get current components (supports PIE)
-	TInlineComponentArray<UDMXFixtureComponent*> DMXComponents;
-	GetComponents<UDMXFixtureComponent>(DMXComponents);
-
-	for (auto& DMXComponent : DMXComponents)
-	{
-		if (DMXComponent->bIsEnabled && DMXComponent->bUseInterpolation && DMXComponent->IsRegistered())
-		{
-			for (auto& Cell : DMXComponent->Cells)
-			{
-				DMXComponent->CurrentCell = &Cell;
-				for (auto& ChannelInterp : Cell.ChannelInterpolation)
-				{
-					if (ChannelInterp.IsUpdating)
-					{
-						// Update
-						ChannelInterp.Travel(DeltaSeconds);
-
-						// Run BP event
-						DMXComponent->InterpolateComponent(DeltaSeconds);
-
-						// Is done?
-						if (ChannelInterp.IsInterpolationDone())
-						{
-							ChannelInterp.EndInterpolation();
-						}
-					}
-				}
-			}
-		}
-	}
-}
-
 void ADMXFixtureActor::SetLightIntensityMax(float NewLightIntensityMax)
 {
 	LightIntensityMax = NewLightIntensityMax;
@@ -305,66 +263,6 @@ void ADMXFixtureActor::PushNormalizedValuesPerAttribute(const FDMXNormalizedAttr
 
 	if (HasBeenInitialized)
 	{
-		for (UDMXFixtureComponent* DMXComponent : TInlineComponentArray<UDMXFixtureComponent*>(this))
-		{
-			// Components without matrix data
-			if (DMXComponent->bIsEnabled)
-			{
-				if (UDMXFixtureComponentSingle* SingleComponent = Cast<UDMXFixtureComponentSingle>(DMXComponent))
-				{
-					// SingleChannel Component
-					const float* TargetValuePtr = ValuePerAttribute.Map.Find(SingleComponent->DMXChannel.Name);
-					if (TargetValuePtr)
-					{
-						const float RemappedValue = SingleComponent->NormalizedToAbsoluteValue(*TargetValuePtr);
-						SingleComponent->SetTargetValue(RemappedValue);
-					}
-				}
-				else if (UDMXFixtureComponentDouble* DoubleComponent = Cast<UDMXFixtureComponentDouble>(DMXComponent))
-				{	
-					// DoubleChannel Component
-					const float* FirstTargetValuePtr = ValuePerAttribute.Map.Find(DoubleComponent->DMXChannel1.Name);
-					if (FirstTargetValuePtr)
-					{
-						float Channel1RemappedValue = DoubleComponent->NormalizedToAbsoluteValue(0, *FirstTargetValuePtr);
-
-						constexpr int32 FirstChannelIndex = 0;
-						DoubleComponent->SetTargetValue(FirstChannelIndex, *FirstTargetValuePtr);
-					}
-
-					const float* SecondTargetValuePtr = ValuePerAttribute.Map.Find(DoubleComponent->DMXChannel2.Name);
-					if (SecondTargetValuePtr)
-					{
-						float Channel2RemappedValue = DoubleComponent->NormalizedToAbsoluteValue(1, *SecondTargetValuePtr);
-
-						constexpr int32 SecondChannelIndex = 1;
-						DoubleComponent->SetTargetValue(SecondChannelIndex, *SecondTargetValuePtr);
-					}
-				}
-				else if(UDMXFixtureComponentColor* ColorComponent = Cast<UDMXFixtureComponentColor>(DMXComponent))
-				{
-					// Color Component
-					if (FLinearColor* CurrentTargetColorPtr = ColorComponent->CurrentTargetColorRef)
-					{
-						const float* FirstTargetValuePtr = ValuePerAttribute.Map.Find(ColorComponent->DMXChannel1);
-						const float* SecondTargetValuePtr = ValuePerAttribute.Map.Find(ColorComponent->DMXChannel2);
-						const float* ThirdTargetValuePtr = ValuePerAttribute.Map.Find(ColorComponent->DMXChannel3);
-						const float* FourthTargetValuePtr = ValuePerAttribute.Map.Find(ColorComponent->DMXChannel4);
-
-						// Current color if channel not found
-						const float r = (FirstTargetValuePtr) ? *FirstTargetValuePtr : CurrentTargetColorPtr->R;
-						const float g = (SecondTargetValuePtr) ? *SecondTargetValuePtr : CurrentTargetColorPtr->G;
-						const float b = (ThirdTargetValuePtr) ? *ThirdTargetValuePtr : CurrentTargetColorPtr->B;
-						const float a = (FourthTargetValuePtr) ? *FourthTargetValuePtr : CurrentTargetColorPtr->A;
-
-						FLinearColor NewTargetColor(r, g, b, a);
-						if (ColorComponent->IsColorValid(NewTargetColor))
-						{
-							ColorComponent->SetTargetColor(NewTargetColor);
-						}
-					}
-				}
-			}
-		}
+		Super::PushNormalizedValuesPerAttribute(ValuePerAttribute);
 	}
 }
diff --git a/Engine/Plugins/VirtualProduction/DMX/DMXFixtures/Source/DMXFixtures/Private/DMXFixtureActorBase.cpp b/Engine/Plugins/VirtualProduction/DMX/DMXFixtures/Source/DMXFixtures/Private/DMXFixtureActorBase.cpp
new file mode 100644
--- /dev/null
+++ b/Engine/Plugins/VirtualProduction/DMX/DMXFixtures/Source/DMXFixtures/Private/DMXFixtureActorBase.cpp
@@ -0,0 +1,75 @@
+// Copyright Epic Games, Inc. All Rights Reserved.
+
+#include "DMXFixtureActorBase.h"
+
+#include "DMXFixtureComponent.h"
+#include "DMXFixtureComponentDouble.h"
+#include "DMXFixtureComponentSingle.h"
+#include "DMXFixtureComponentColor.h"
+#include "DMXStats.h"
+#include "DMXTypes.h"
+#include "Game/DMXComponent.h"
+
+#include "Components/StaticMeshComponent.h"
+
+
+DECLARE_CYCLE_STAT(TEXT("BaseFixtureActor Push Normalized Values"), STAT_BaseFixtureActorPushNormalizedValuesPerAttribute, STATGROUP_DMX);
+
+ADMXFixtureActorBase::ADMXFixtureActorBase()
+{
+	PrimaryActorTick.bCanEverTick = true;
+
+#if WITH_EDITORONLY_DATA
+	bRunConstructionScriptOnDrag = false;
+#endif
+
+	DMX = CreateDefaultSubobject<UDMXComponent>(TEXT("DMX"));
+}
+
+void ADMXFixtureActorBase::InterpolateDMXComponents(float DeltaSeconds)
+{
+	// Get current components (supports PIE)
+	TInlineComponentArray<UDMXFixtureComponent*> DMXComponents;
+	GetComponents<UDMXFixtureComponent>(DMXComponents);
+
+	for (auto& DMXComponent : DMXComponents)
+	{
+		if (DMXComponent->bIsEnabled && DMXComponent->bUseInterpolation && DMXComponent->IsRegistered())
+		{
+			for (auto& Cell : DMXComponent->Cells)
+			{
+				DMXComponent->CurrentCell = &Cell;
+				for (auto& ChannelInterp : Cell.ChannelInterpolation)
+				{
+					if (ChannelInterp.IsUpdating)
+					{
+						// Update
+						ChannelInterp.Travel(DeltaSeconds);
+
+						// Run BP event
+						DMXComponent->InterpolateComponent(DeltaSeconds);
+
+						// Is done?
+						if (ChannelInterp.IsInterpolationDone())
+						{
+							ChannelInterp.EndInterpolation();
+						}
+					}
+				}
+			}
+		}
+	}
+}
+
+void ADMXFixtureActorBase::PushNormalizedValuesPerAttribute(const FDMXNormalizedAttributeValueMap& ValuePerAttribute)
+{
+	SCOPE_CYCLE_COUNTER(STAT_BaseFixtureActorPushNormalizedValuesPerAttribute);
+
+	for (UDMXFixtureComponent* DMXComponent : TInlineComponentArray<UDMXFixtureComponent*>(this))
+	{
+		if (DMXComponent->bIsEnabled)
+		{
+			DMXComponent->PushNormalizedValuesPerAttribute(ValuePerAttribute);
+		}
+	}
+}
diff --git a/Engine/Plugins/VirtualProduction/DMX/DMXFixtures/Source/DMXFixtures/Private/DMXFixtureActorMatrix.cpp b/Engine/Plugins/VirtualProduction/DMX/DMXFixtures/Source/DMXFixtures/Private/DMXFixtureActorMatrix.cpp
--- a/Engine/Plugins/VirtualProduction/DMX/DMXFixtures/Source/DMXFixtures/Private/DMXFixtureActorMatrix.cpp
+++ b/Engine/Plugins/VirtualProduction/DMX/DMXFixtures/Source/DMXFixtures/Private/DMXFixtureActorMatrix.cpp
@@ -4,6 +4,7 @@
 
 #include "DMXProtocolConstants.h"
 #include "DMXStats.h"
+#include "Game/DMXComponent.h"
 
 #include "Engine/Texture2D.h"
 #include "Rendering/Texture2DResource.h"
diff --git a/Engine/Plugins/VirtualProduction/DMX/DMXFixtures/Source/DMXFixtures/Private/DMXFixtureComponentColor.cpp b/Engine/Plugins/VirtualProduction/DMX/DMXFixtures/Source/DMXFixtures/Private/DMXFixtureComponentColor.cpp
--- a/Engine/Plugins/VirtualProduction/DMX/DMXFixtures/Source/DMXFixtures/Private/DMXFixtureComponentColor.cpp
+++ b/Engine/Plugins/VirtualProduction/DMX/DMXFixtures/Source/DMXFixtures/Private/DMXFixtureComponentColor.cpp
@@ -2,6 +2,8 @@
 
 #include "DMXFixtureComponentColor.h"
 
+#include "DMXTypes.h"
+
 
 UDMXFixtureComponentColor::UDMXFixtureComponentColor()
 	: CurrentTargetColorRef(nullptr)
@@ -9,6 +11,29 @@ UDMXFixtureComponentColor::UDMXFixtureComponentColor()
 	PrimaryComponentTick.bCanEverTick = false;
 }
 
+void UDMXFixtureComponentColor::PushNormalizedValuesPerAttribute(const FDMXNormalizedAttributeValueMap& ValuePerAttribute)
+{
+	if (FLinearColor* CurrentTargetColorPtr = CurrentTargetColorRef)
+	{
+		const float* FirstTargetValuePtr = ValuePerAttribute.Map.Find(DMXChannel1);
+		const float* SecondTargetValuePtr = ValuePerAttribute.Map.Find(DMXChannel2);
+		const float* ThirdTargetValuePtr = ValuePerAttribute.Map.Find(DMXChannel3);
+		const float* FourthTargetValuePtr = ValuePerAttribute.Map.Find(DMXChannel4);
+
+		// Current color if channel not found
+		const float r = (FirstTargetValuePtr) ? *FirstTargetValuePtr : CurrentTargetColorPtr->R;
+		const float g = (SecondTargetValuePtr) ? *SecondTargetValuePtr : CurrentTargetColorPtr->G;
+		const float b = (ThirdTargetValuePtr) ? *ThirdTargetValuePtr : CurrentTargetColorPtr->B;
+		const float a = (FourthTargetValuePtr) ? *FourthTargetValuePtr : CurrentTargetColorPtr->A;
+
+		FLinearColor NewTargetColor(r, g, b, a);
+		if (IsColorValid(NewTargetColor))
+		{
+			SetTargetColor(NewTargetColor);
+		}
+	}
+}
+
 void UDMXFixtureComponentColor::GetSupportedDMXAttributes_Implementation(TArray<FName>& OutAttributeNames)
 {
 	if (bIsEnabled)
diff --git a/Engine/Plugins/VirtualProduction/DMX/DMXFixtures/Source/DMXFixtures/Private/DMXFixtureComponentDouble.cpp b/Engine/Plugins/VirtualProduction/DMX/DMXFixtures/Source/DMXFixtures/Private/DMXFixtureComponentDouble.cpp
--- a/Engine/Plugins/VirtualProduction/DMX/DMXFixtures/Source/DMXFixtures/Private/DMXFixtureComponentDouble.cpp
+++ b/Engine/Plugins/VirtualProduction/DMX/DMXFixtures/Source/DMXFixtures/Private/DMXFixtureComponentDouble.cpp
@@ -4,6 +4,7 @@
 #include "DMXFixtureComponentDouble.h"
 #include "DMXFixtureActor.h"
 
+#include "DMXTypes.h"
 
 
 UDMXFixtureComponentDouble::UDMXFixtureComponentDouble()
@@ -11,6 +12,27 @@ UDMXFixtureComponentDouble::UDMXFixtureComponentDouble()
 	PrimaryComponentTick.bCanEverTick = false;
 }
 
+void UDMXFixtureComponentDouble::PushNormalizedValuesPerAttribute(const FDMXNormalizedAttributeValueMap& ValuePerAttribute)
+{
+	const float* FirstTargetValuePtr = ValuePerAttribute.Map.Find(DMXChannel1.Name);
+	if (FirstTargetValuePtr)
+	{
+		float Channel1RemappedValue = NormalizedToAbsoluteValue(0, *FirstTargetValuePtr);
+
+		constexpr int32 FirstChannelIndex = 0;
+		SetTargetValue(FirstChannelIndex, *FirstTargetValuePtr);
+	}
+
+	const float* SecondTargetValuePtr = ValuePerAttribute.Map.Find(DMXChannel2.Name);
+	if (SecondTargetValuePtr)
+	{
+		float Channel2RemappedValue = NormalizedToAbsoluteValue(1, *SecondTargetValuePtr);
+
+		constexpr int32 SecondChannelIndex = 1;
+		SetTargetValue(SecondChannelIndex, *SecondTargetValuePtr);
+	}
+}
+
 void UDMXFixtureComponentDouble::GetSupportedDMXAttributes_Implementation(TArray<FName>& OutAttributeNames)
 {
 	if (bIsEnabled)
diff --git a/Engine/Plugins/VirtualProduction/DMX/DMXFixtures/Source/DMXFixtures/Private/DMXFixtureComponentSingle.cpp b/Engine/Plugins/VirtualProduction/DMX/DMXFixtures/Source/DMXFixtures/Private/DMXFixtureComponentSingle.cpp
--- a/Engine/Plugins/VirtualProduction/DMX/DMXFixtures/Source/DMXFixtures/Private/DMXFixtureComponentSingle.cpp
+++ b/Engine/Plugins/VirtualProduction/DMX/DMXFixtures/Source/DMXFixtures/Private/DMXFixtureComponentSingle.cpp
@@ -4,12 +4,24 @@
 #include "DMXFixtureComponentSingle.h"
 #include "DMXFixtureActor.h"
 
+#include "DMXTypes.h"
+
 
 UDMXFixtureComponentSingle::UDMXFixtureComponentSingle()
 {
 	PrimaryComponentTick.bCanEverTick = false;
 }
 
+void UDMXFixtureComponentSingle::PushNormalizedValuesPerAttribute(const FDMXNormalizedAttributeValueMap& ValuePerAttribute)
+{
+	const float* TargetValuePtr = ValuePerAttribute.Map.Find(DMXChannel.Name);
+	if (TargetValuePtr)
+	{
+		const float RemappedValue = NormalizedToAbsoluteValue(*TargetValuePtr);
+		SetTargetValue(RemappedValue);
+	}
+}
+
 void UDMXFixtureComponentSingle::GetSupportedDMXAttributes_Implementation(TArray<FName>& OutAttributeNames)
 {
 	if (bIsEnabled)
diff --git a/Engine/Plugins/VirtualProduction/DMX/DMXFixtures/Source/DMXFixtures/Public/DMXFixtureActor.h b/Engine/Plugins/VirtualProduction/DMX/DMXFixtures/Source/DMXFixtures/Public/DMXFixtureActor.h
--- a/Engine/Plugins/VirtualProduction/DMX/DMXFixtures/Source/DMXFixtures/Public/DMXFixtureActor.h
+++ b/Engine/Plugins/VirtualProduction/DMX/DMXFixtures/Source/DMXFixtures/Public/DMXFixtureActor.h
@@ -2,6 +2,8 @@
 
 #pragma once
 
+#include "DMXFixtureActorBase.h"
+
 #include "DMXFixtureComponent.h"
 #include "DMXFixtureComponentDouble.h"
 #include "DMXFixtureComponentSingle.h"
@@ -13,8 +15,6 @@
 #include "CoreMinimal.h"
 #include "Components/ArrowComponent.h"
 #include "Components/SpotLightComponent.h"
-#include "Game/DMXComponent.h"
-#include "GameFramework/Actor.h"
 #include "Materials/MaterialInstanceDynamic.h"
 
 #include "DMXFixtureActor.generated.h"
@@ -32,7 +32,7 @@ enum EDMXFixtureQualityLevel
 
 UCLASS()
 class DMXFIXTURES_API ADMXFixtureActor 
-	: public AActor
+	: public ADMXFixtureActorBase
 	, public IDMXMVRFixtureActorInterface
 {
 	GENERATED_BODY()
@@ -87,13 +87,9 @@ public:
 	void InitializeFixture(UStaticMeshComponent* StaticMeshLens, UStaticMeshComponent* StaticMeshBeam);
 
 	/** Pushes DMX Values to the Fixture. Expects normalized values in the range of 0.0f - 1.0f */
-	UFUNCTION(BlueprintCallable, Category = "DMX Fixture")
-	void PushNormalizedValuesPerAttribute(const FDMXNormalizedAttributeValueMap& ValuePerAttributeMap);
+	virtual void PushNormalizedValuesPerAttribute(const FDMXNormalizedAttributeValueMap& ValuePerAttributeMap) override;
 	
 public:
-	UFUNCTION(BlueprintCallable, Category = "DMX Fixture")
-	void InterpolateDMXComponents(float DeltaSeconds);
-	
 	/** Sets the a new max light intensity */
 	UFUNCTION(BlueprintCallable, Category = "DMX Fixture")
 	void SetLightIntensityMax(float NewLightIntensityMax);
@@ -149,12 +145,6 @@ public:
 	UPROPERTY(EditAnywhere, BlueprintReadWrite, Category = "DMX Light Fixture")
 	bool UseDynamicOcclusion;
 
-
-	// DMX COMPONENT -----------------------------
-
-	UPROPERTY(VisibleAnywhere, BlueprintReadOnly, Category = "DMX Light Fixture")
-	TObjectPtr<class UDMXComponent> DMX;
-
 	// COMPONENTS ---------------------------------
 
 	UPROPERTY(BlueprintReadOnly, Category = "DMX Light Fixture")
diff --git a/Engine/Plugins/VirtualProduction/DMX/DMXFixtures/Source/DMXFixtures/Public/DMXFixtureActorBase.h b/Engine/Plugins/VirtualProduction/DMX/DMXFixtures/Source/DMXFixtures/Public/DMXFixtureActorBase.h
new file mode 100644
--- /dev/null
+++ b/Engine/Plugins/VirtualProduction/DMX/DMXFixtures/Source/DMXFixtures/Public/DMXFixtureActorBase.h
@@ -0,0 +1,31 @@
+// Copyright Epic Games, Inc. All Rights Reserved.
+
+#pragma once
+
+#include "GameFramework/Actor.h"
+
+#include "DMXFixtureActorBase.generated.h"
+
+struct FDMXNormalizedAttributeValueMap;
+class UDMXComponent;
+
+
+UCLASS()
+class DMXFIXTURES_API ADMXFixtureActorBase
+	: public AActor
+{
+	GENERATED_BODY()
+
+public:
+	ADMXFixtureActorBase();
+
+	/** Pushes DMX Values to the Fixture. Expects normalized values in the range of 0.0f - 1.0f */
+	UFUNCTION(BlueprintCallable, Category = "DMX Fixture")
+	virtual void PushNormalizedValuesPerAttribute(const FDMXNormalizedAttributeValueMap& ValuePerAttributeMap);
+
+	UFUNCTION(BlueprintCallable, Category = "DMX Fixture")
+	void InterpolateDMXComponents(float DeltaSeconds);
+
+	UPROPERTY(VisibleAnywhere, BlueprintReadOnly, Category = "DMX Light Fixture")
+	TObjectPtr<UDMXComponent> DMX;
+};
diff --git a/Engine/Plugins/VirtualProduction/DMX/DMXFixtures/Source/DMXFixtures/Public/DMXFixtureComponent.h b/Engine/Plugins/VirtualProduction/DMX/DMXFixtures/Source/DMXFixtures/Public/DMXFixtureComponent.h
--- a/Engine/Plugins/VirtualProduction/DMX/DMXFixtures/Source/DMXFixtures/Public/DMXFixtureComponent.h
+++ b/Engine/Plugins/VirtualProduction/DMX/DMXFixtures/Source/DMXFixtures/Public/DMXFixtureComponent.h
@@ -12,6 +12,8 @@
 #include "DMXInterpolation.h"
 #include "DMXFixtureComponent.generated.h"
 
+struct FDMXNormalizedAttributeValueMap;
+
 
 USTRUCT(BlueprintType)
 struct FDMXChannelData
@@ -32,7 +34,7 @@ struct FDMXChannelData
 };
 
 
-UCLASS(meta=(IsBlueprintBase=false))
+UCLASS(Abstract, Meta=(IsBlueprintBase=false))
 class DMXFIXTURES_API UDMXFixtureComponent : public UActorComponent
 {
 	GENERATED_BODY()
@@ -51,6 +53,9 @@ public:
 	/** Sets the cell that is currently active */
 	virtual void SetCurrentCell(int Index);
 
+	/** Pushes DMX Values to the Fixture Component. Expects normalized values in the range of 0.0f to 1.0f */
+	virtual void PushNormalizedValuesPerAttribute(const FDMXNormalizedAttributeValueMap& ValuePerAttribute)  PURE_VIRTUAL(UDMXFixtureComponent::PushNormalizedValuesPerAttribute, return; );
+
 	/** If used within a DMX Fixture Actor or Fixture Matrix Actor, the component only receives data when set to true. Else needs be implemented in blueprints. */
 	UPROPERTY(EditAnywhere, Category = "DMX Parameters", meta = (DisplayPriority = 0))
 	bool bIsEnabled = true;
diff --git a/Engine/Plugins/VirtualProduction/DMX/DMXFixtures/Source/DMXFixtures/Public/DMXFixtureComponentColor.h b/Engine/Plugins/VirtualProduction/DMX/DMXFixtures/Source/DMXFixtures/Public/DMXFixtureComponentColor.h
--- a/Engine/Plugins/VirtualProduction/DMX/DMXFixtures/Source/DMXFixtures/Public/DMXFixtureComponentColor.h
+++ b/Engine/Plugins/VirtualProduction/DMX/DMXFixtures/Source/DMXFixtures/Public/DMXFixtureComponentColor.h
@@ -19,6 +19,7 @@ public:
 	UDMXFixtureComponentColor();
 
 	//~ Begin DMXFixtureComponent interface
+	virtual void PushNormalizedValuesPerAttribute(const FDMXNormalizedAttributeValueMap& ValuePerAttribute) override;
 	virtual void GetSupportedDMXAttributes_Implementation(TArray<FName>& OutAttributeNames) override;
 	//~ End DMXFixtureComponent interface
 
diff --git a/Engine/Plugins/VirtualProduction/DMX/DMXFixtures/Source/DMXFixtures/Public/DMXFixtureComponentDouble.h b/Engine/Plugins/VirtualProduction/DMX/DMXFixtures/Source/DMXFixtures/Public/DMXFixtureComponentDouble.h
--- a/Engine/Plugins/VirtualProduction/DMX/DMXFixtures/Source/DMXFixtures/Public/DMXFixtureComponentDouble.h
+++ b/Engine/Plugins/VirtualProduction/DMX/DMXFixtures/Source/DMXFixtures/Public/DMXFixtureComponentDouble.h
@@ -16,6 +16,7 @@ public:
 	UDMXFixtureComponentDouble();
 
 	//~ Begin DMXFixtureComponent interface
+	virtual void PushNormalizedValuesPerAttribute(const FDMXNormalizedAttributeValueMap& ValuePerAttribute) override;
 	virtual void GetSupportedDMXAttributes_Implementation(TArray<FName>& OutAttributeNames) override;
 	//~ End DMXFixtureComponent interface
 
diff --git a/Engine/Plugins/VirtualProduction/DMX/DMXFixtures/Source/DMXFixtures/Public/DMXFixtureComponentSingle.h b/Engine/Plugins/VirtualProduction/DMX/DMXFixtures/Source/DMXFixtures/Public/DMXFixtureComponentSingle.h
--- a/Engine/Plugins/VirtualProduction/DMX/DMXFixtures/Source/DMXFixtures/Public/DMXFixtureComponentSingle.h
+++ b/Engine/Plugins/VirtualProduction/DMX/DMXFixtures/Source/DMXFixtures/Public/DMXFixtureComponentSingle.h
@@ -16,6 +16,7 @@ public:
 	UDMXFixtureComponentSingle();
 
 	//~ Begin DMXFixtureComponent interface
+	virtual void PushNormalizedValuesPerAttribute(const FDMXNormalizedAttributeValueMap& ValuePerAttribute) override;
 	virtual void GetSupportedDMXAttributes_Implementation(TArray<FName>& OutAttributeNames) override;
 	//~ End DMXFixtureComponent interface
 		
