DMX - Fixture Matrix Actor now alows to call InitializeFixture more than once

diff --git a/Engine/Plugins/VirtualProduction/DMX/DMXFixtures/Source/DMXFixtures/Private/DMXFixtureActorMatrix.cpp b/Engine/Plugins/VirtualProduction/DMX/DMXFixtures/Source/DMXFixtures/Private/DMXFixtureActorMatrix.cpp
--- a/Engine/Plugins/VirtualProduction/DMX/DMXFixtures/Source/DMXFixtures/Private/DMXFixtureActorMatrix.cpp
+++ b/Engine/Plugins/VirtualProduction/DMX/DMXFixtures/Source/DMXFixtures/Private/DMXFixtureActorMatrix.cpp
@@ -14,7 +14,7 @@ DECLARE_CYCLE_STAT(TEXT("Fixture Actor Matrix Push Fixture Matrix Cell Data"), S
 
 namespace
 {
-	void UpdateMatrixTexture(uint8* MatrixData, UTexture2D* DynamicTexture, int32 MipIndex, uint32 NumRegions, FUpdateTextureRegion2D Region, uint32 SrcPitch, uint32 SrcBpp)
+	void UpdateMatrixTexture(uint8* MatrixData, UTexture2D* DynamicTexture, int32 MipIndex, uint32 NumRegions, const FUpdateTextureRegion2D& Region, uint32 SrcPitch, uint32 SrcBpp)
 	{
 		if (DynamicTexture->GetResource())
 		{
@@ -54,14 +54,6 @@ ADMXFixtureActorMatrix::ADMXFixtureActorMatrix()
 	XCells = 1;
 	YCells = 1;
 	MatrixDataSize = 0;
-
-	TextureRegion = nullptr;
-}
-
-ADMXFixtureActorMatrix::~ADMXFixtureActorMatrix()
-{
-	if (TextureRegion)
-		delete TextureRegion;
 }
 
 #if WITH_EDITOR
@@ -169,8 +161,7 @@ void ADMXFixtureActorMatrix::InitializeMatrixFixture()
 		MatrixDataTexture->CompressionSettings = TextureCompressionSettings::TC_VectorDisplacementmap;
 		MatrixDataTexture->UpdateResource(); //to initialize resource
 
-		check(!TextureRegion);
-		TextureRegion = new FUpdateTextureRegion2D(0, 0, 0, 0, TextureWidth, TextureHeight);
+		TextureRegion = FUpdateTextureRegion2D(0, 0, 0, 0, TextureWidth, TextureHeight);
 
 		// Push fixture data into materials and lights
 		FeedFixtureData();
@@ -359,7 +350,7 @@ void ADMXFixtureActorMatrix::UpdateDynamicTexture()
 	if (MatrixDataTexture)
 	{
 		int NbrCells = XCells * YCells;
-		UpdateMatrixTexture(MatrixData.GetData(), MatrixDataTexture, 0, 1, *TextureRegion, NbrCells * 4, 4);
+		UpdateMatrixTexture(MatrixData.GetData(), MatrixDataTexture, 0, 1, TextureRegion, NbrCells * 4, 4);
 	}
 }
 
diff --git a/Engine/Plugins/VirtualProduction/DMX/DMXFixtures/Source/DMXFixtures/Public/DMXFixtureActorMatrix.h b/Engine/Plugins/VirtualProduction/DMX/DMXFixtures/Source/DMXFixtures/Public/DMXFixtureActorMatrix.h
--- a/Engine/Plugins/VirtualProduction/DMX/DMXFixtures/Source/DMXFixtures/Public/DMXFixtureActorMatrix.h
+++ b/Engine/Plugins/VirtualProduction/DMX/DMXFixtures/Source/DMXFixtures/Public/DMXFixtureActorMatrix.h
@@ -3,10 +3,14 @@
 #pragma once
 
 #include "DMXFixtureActor.h"
+
 #include "ProceduralMeshComponent.h"
+#include "RHI.h"
 #include "Rendering/Texture2DResource.h"
+
 #include "DMXFixtureActorMatrix.generated.h"
 
+
 UCLASS()
 class DMXFIXTURES_API ADMXFixtureActorMatrix : public ADMXFixtureActor
 {
@@ -14,7 +18,6 @@ class DMXFIXTURES_API ADMXFixtureActorMatrix : public ADMXFixtureActor
 
 public:
 	ADMXFixtureActorMatrix();
-	~ADMXFixtureActorMatrix();
 
 protected:
 	
@@ -35,7 +38,7 @@ public:
 	int YCells;
 
 	UTexture2D* MatrixDataTexture;
-	FUpdateTextureRegion2D* TextureRegion;
+	FUpdateTextureRegion2D TextureRegion;
 
 	UPROPERTY(EditDefaultsOnly, BlueprintReadOnly, Category = "DMX Matrix Fixture")
 	TObjectPtr<UProceduralMeshComponent> MatrixHead;
