DMXFixtures - Deprecated MinQuality and MaxQuality UProperties. Now using BeamQuality and ZoomQuality.

diff --git a/Engine/Plugins/VirtualProduction/DMX/DMXFixtures/Source/DMXFixtures/Private/DMXFixtureActor.cpp b/Engine/Plugins/VirtualProduction/DMX/DMXFixtures/Source/DMXFixtures/Private/DMXFixtureActor.cpp
--- a/Engine/Plugins/VirtualProduction/DMX/DMXFixtures/Source/DMXFixtures/Private/DMXFixtureActor.cpp
+++ b/Engine/Plugins/VirtualProduction/DMX/DMXFixtures/Source/DMXFixtures/Private/DMXFixtureActor.cpp
@@ -44,11 +44,24 @@ ADMXFixtureActor::ADMXFixtureActor()
 	UseDynamicOcclusion = false;
 	LensRadius = 10.0f;
 	QualityLevel = EDMXFixtureQualityLevel::HighQuality;
-	MinQuality = 1.0f;
-	MaxQuality = 1.0f;
+	ZoomQuality = 1.0f;
+	BeamQuality = 1.0f;
+	DisableLights = false;
 	HasBeenInitialized = false;
 }
 
+
+void ADMXFixtureActor::PostLoad()
+{
+	Super::PostLoad();
+
+	// Backward compatibility with new UPROPERTYs
+	ZoomQuality = MaxQuality;
+	BeamQuality = MinQuality;
+}
+
+
+
 void ADMXFixtureActor::OnMVRGetSupportedDMXAttributes_Implementation(TArray<FName>& OutAttributeNames, TArray<FName>& OutMatrixAttributeNames) const
 {
 	for (UDMXFixtureComponent* DMXFixtureComponent : TInlineComponentArray<UDMXFixtureComponent*>(this))
@@ -135,35 +148,38 @@ void ADMXFixtureActor::InitializeFixture(UStaticMeshComponent* StaticMeshLens, U
 
 void ADMXFixtureActor::FeedFixtureData()
 {
-	// Note: MinQuality and MaxQuality are used in conjonction with the zoom angle when zoom component is used
-	switch (QualityLevel)
-	{
-		case(EDMXFixtureQualityLevel::LowQuality): MinQuality = 4.0f; MaxQuality = 4.0f; break;
-		case(EDMXFixtureQualityLevel::MediumQuality): MinQuality = 2.0f; MaxQuality = 2.0f; break;
-		case(EDMXFixtureQualityLevel::HighQuality): MinQuality = 1.0f; MaxQuality = 1.0f; break;
-		case(EDMXFixtureQualityLevel::UltraQuality): MinQuality = 0.33f; MaxQuality = 0.33f; break;
-	}
-
-	// Note:fallback when fixture doesnt use zoom component
-	float QualityFallback = 1.0f;
+	// BeamQuality and ZoomQuality modulate the "stepSize" for the raymarch beam shader
+	// lower value is visually better
 	switch (QualityLevel)
 	{
-		case(EDMXFixtureQualityLevel::LowQuality): QualityFallback = 4.0f; break;
-		case(EDMXFixtureQualityLevel::MediumQuality): QualityFallback = 2.0f; break;
-		case(EDMXFixtureQualityLevel::HighQuality): QualityFallback = 1.0f; break;
-		case(EDMXFixtureQualityLevel::UltraQuality): QualityFallback = 0.33f; break;
+		case(EDMXFixtureQualityLevel::LowQuality): 
+			ZoomQuality = 1.0f;
+			BeamQuality = 4.0f;
+			break;
+		case(EDMXFixtureQualityLevel::MediumQuality): 
+			ZoomQuality = 1.0f;
+			BeamQuality = 2.0f;
+			break;
+		case(EDMXFixtureQualityLevel::HighQuality): 
+			ZoomQuality = 1.0f;
+			BeamQuality = 1.0f;
+			break;
+		case(EDMXFixtureQualityLevel::UltraQuality): 
+			ZoomQuality = 2.0f;
+			BeamQuality = 0.33f;
+			break;
 	}
 	
+	// Clamp values, also clamps UI
 	if (QualityLevel == EDMXFixtureQualityLevel::Custom)
 	{
-		MinQuality = FMath::Clamp(MinQuality, 0.2f, 4.0f);
-		MaxQuality = FMath::Clamp(MaxQuality, 0.2f, 4.0f);
-		QualityFallback = MaxQuality;
+		ZoomQuality = FMath::Clamp(ZoomQuality, 1.0f, 4.0f);
+		BeamQuality = FMath::Clamp(BeamQuality, 0.2f, 4.0f);
 	}
 
 	if (DynamicMaterialBeam)
 	{
-		DynamicMaterialBeam->SetScalarParameterValue("DMX Quality Level", QualityFallback);
+		DynamicMaterialBeam->SetScalarParameterValue("DMX Quality Level", BeamQuality);
 		DynamicMaterialBeam->SetScalarParameterValue("DMX Max Light Distance", LightDistanceMax);
 		DynamicMaterialBeam->SetScalarParameterValue("DMX Max Light Intensity", LightIntensityMax * SpotlightIntensityScale);
 	}
diff --git a/Engine/Plugins/VirtualProduction/DMX/DMXFixtures/Source/DMXFixtures/Public/DMXFixtureActor.h b/Engine/Plugins/VirtualProduction/DMX/DMXFixtures/Source/DMXFixtures/Public/DMXFixtureActor.h
--- a/Engine/Plugins/VirtualProduction/DMX/DMXFixtures/Source/DMXFixtures/Public/DMXFixtureActor.h
+++ b/Engine/Plugins/VirtualProduction/DMX/DMXFixtures/Source/DMXFixtures/Public/DMXFixtureActor.h
@@ -52,6 +52,9 @@ public:
 	virtual void OnMVRSetFixturePatch_Implementation(UDMXEntityFixturePatch* FixturePatch) override;
 	virtual UDMXEntityFixturePatch* OnMVRGetFixturePatch_Implementation() const override;
 	//~ End DMXMVRFixtureActorInterface interface
+	
+	// UObject interface
+	virtual void PostLoad() override;
 
 	bool HasBeenInitialized;
 	float LensRadius;
@@ -63,13 +66,13 @@ public:
 	UPROPERTY(EditAnywhere, BlueprintReadWrite, Category = "DMX Light Fixture", meta = (DisplayPriority = 0))
 	TEnumAsByte<EDMXFixtureQualityLevel> QualityLevel;
 
-	// Visual quality when using smaller zoom angle (thin beam). Small value is visually better but cost more on GPU
+	// Visual quality for the light beam - small value is visually better but cost more on GPU
 	UPROPERTY(EditAnywhere, BlueprintReadWrite, Category = "DMX Light Fixture", meta = (EditCondition = "QualityLevel == EDMXFixtureQualityLevel::Custom", EditConditionHides))
-	float MinQuality;
+	float BeamQuality;
 
-	// Visual quality when using bigger zoom angle (wide beam). Small value is visually better but cost more on GPU
+	// Visual quality for the light beam when zoom is wide - small value is visually better but cost more on GPU
 	UPROPERTY(EditAnywhere, BlueprintReadWrite, Category = "DMX Light Fixture", meta = (EditCondition = "QualityLevel == EDMXFixtureQualityLevel::Custom", EditConditionHides))
-	float MaxQuality;
+	float ZoomQuality;
 
 	// HIERARCHY---------------------------------
 	UPROPERTY(EditDefaultsOnly, BlueprintReadOnly, Category = "DMX Light Fixture")
@@ -145,6 +148,10 @@ public:
 	UPROPERTY(EditAnywhere, BlueprintReadWrite, Category = "DMX Light Fixture")
 	bool UseDynamicOcclusion;
 
+	// Disable lights rendering in the fixture to save on GPU
+	UPROPERTY(EditAnywhere, BlueprintReadWrite, Category = "DMX Light Fixture")
+	bool DisableLights;
+
 	// COMPONENTS ---------------------------------
 
 	UPROPERTY(BlueprintReadOnly, Category = "DMX Light Fixture")
@@ -186,4 +193,14 @@ public:
 	UPROPERTY(BlueprintReadOnly, Category = "DMX Light Fixture")
 	TObjectPtr<UMaterialInstanceDynamic> DynamicMaterialPointLight;
 
+	///////////////////////////////////////
+	// DEPRECATED 5.1
+public:
+	UE_DEPRECATED(5.1, "Use BeamQuality instead")
+	UPROPERTY(meta = (DeprecatedProperty, DeprecationMessage = "Use BeamQuality instead."))
+	float MinQuality;
+
+	UE_DEPRECATED(5.1, "Use ZoomQuality instead")
+	UPROPERTY(meta = (DeprecatedProperty, DeprecationMessage = "Use ZoomQuality instead."))
+	float MaxQuality;
 };
