DMX - Fix CIS Compile errors in DMXFixtureActor.cpp and 3 others

diff --git a/Engine/Plugins/VirtualProduction/DMX/DMXFixtures/Source/DMXFixtures/Private/DMXFixtureActor.cpp b/Engine/Plugins/VirtualProduction/DMX/DMXFixtures/Source/DMXFixtures/Private/DMXFixtureActor.cpp
--- a/Engine/Plugins/VirtualProduction/DMX/DMXFixtures/Source/DMXFixtures/Private/DMXFixtureActor.cpp
+++ b/Engine/Plugins/VirtualProduction/DMX/DMXFixtures/Source/DMXFixtures/Private/DMXFixtureActor.cpp
@@ -55,9 +55,11 @@ void ADMXFixtureActor::PostLoad()
 {
 	Super::PostLoad();
 
-	// Backward compatibility with new UPROPERTYs
+	PRAGMA_DISABLE_DEPRECATION_WARNINGS
+	// Upgrade from deprecated MinQuality and MaxQuality to ZoomQuality and BeamQuality
 	ZoomQuality = MaxQuality;
 	BeamQuality = MinQuality;
+	PRAGMA_ENABLE_DEPRECATION_WARNINGS
 }
 
 
