DMX - DMX fixture component's GetTextureCenterColors now has an option to update the texture resource, color wheel colors are now properly initialized at startup, remove broken upgrade path for quality settings

diff --git a/Engine/Plugins/VirtualProduction/DMX/DMXFixtures/Source/DMXFixtures/Private/DMXFixtureActor.cpp b/Engine/Plugins/VirtualProduction/DMX/DMXFixtures/Source/DMXFixtures/Private/DMXFixtureActor.cpp
--- a/Engine/Plugins/VirtualProduction/DMX/DMXFixtures/Source/DMXFixtures/Private/DMXFixtureActor.cpp
+++ b/Engine/Plugins/VirtualProduction/DMX/DMXFixtures/Source/DMXFixtures/Private/DMXFixtureActor.cpp
@@ -3,11 +3,10 @@
 #include "DMXFixtureActor.h"
 
 #include "Components/ArrowComponent.h"
-#include "DMXStats.h"
-
 #include "Components/SpotLightComponent.h"
 #include "Components/StaticMeshComponent.h"
 #include "DMXFixtureComponent.h"
+#include "DMXStats.h"
 #include "Materials/MaterialInstanceDynamic.h"
 
 
@@ -53,17 +52,6 @@ ADMXFixtureActor::ADMXFixtureActor()
 	HasBeenInitialized = false;
 }
 
-void ADMXFixtureActor::PostLoad()
-{
-	Super::PostLoad();
-
-	PRAGMA_DISABLE_DEPRECATION_WARNINGS
-	// Upgrade from deprecated MinQuality and MaxQuality to ZoomQuality and BeamQuality
-	ZoomQuality = MaxQuality;
-	BeamQuality = MinQuality;
-	PRAGMA_ENABLE_DEPRECATION_WARNINGS
-}
-
 void ADMXFixtureActor::OnMVRGetSupportedDMXAttributes_Implementation(TArray<FName>& OutAttributeNames, TArray<FName>& OutMatrixAttributeNames) const
 {
 	for (UDMXFixtureComponent* DMXFixtureComponent : TInlineComponentArray<UDMXFixtureComponent*>(this))
@@ -79,7 +67,11 @@ void ADMXFixtureActor::OnMVRGetSupportedDMXAttributes_Implementation(TArray<FNam
 void ADMXFixtureActor::PostEditChangeProperty(FPropertyChangedEvent& PropertyChangedEvent)
 {
 	Super::PostEditChangeProperty(PropertyChangedEvent);
-	FeedFixtureData();
+
+	if (PropertyChangedEvent.ChangeType != EPropertyChangeType::Interactive)
+	{
+		FeedFixtureData();
+	}
 }
 #endif
 
@@ -117,7 +109,7 @@ void ADMXFixtureActor::InitializeFixture(UStaticMeshComponent* StaticMeshLens, U
 	if (StaticMeshBeam)
 	{
 		StaticMeshBeam->SetMaterial(0, DynamicMaterialBeam);
-		StaticMeshBeam->SetWorldScale3D(FVector(1,1,1));
+		StaticMeshBeam->SetWorldScale3D(FVector(1, 1, 1));
 	}
 
 	// Assign dynamic materials to lights
diff --git a/Engine/Plugins/VirtualProduction/DMX/DMXFixtures/Source/DMXFixtures/Private/DMXFixtureComponent.cpp b/Engine/Plugins/VirtualProduction/DMX/DMXFixtures/Source/DMXFixtures/Private/DMXFixtureComponent.cpp
--- a/Engine/Plugins/VirtualProduction/DMX/DMXFixtures/Source/DMXFixtures/Private/DMXFixtureComponent.cpp
+++ b/Engine/Plugins/VirtualProduction/DMX/DMXFixtures/Source/DMXFixtures/Private/DMXFixtureComponent.cpp
@@ -84,12 +84,17 @@ ADMXFixtureActor* UDMXFixtureComponent::GetParentFixtureActor()
 
 // Warning: Input Texture must use VectorDisplacementMap compression or else RHI complains about BytesPerPixel
 // Reads pixel color in the middle of each "Texture" and output linear colors
-TArray<FLinearColor> UDMXFixtureComponent::GetTextureCenterColors(UTexture2D* Texture, int NumCells)
+TArray<FLinearColor> UDMXFixtureComponent::GetTextureCenterColors(UTexture2D* Texture, int NumCells, bool bUpdateTextureResource)
 {
 	TArray<FLinearColor> PixelColorArray;
 	PixelColorArray.SetNumZeroed(NumCells, true);
 	if (Texture)
 	{
+		if (bUpdateTextureResource)
+		{
+			Texture->UpdateResource();
+		}
+
 		FRenderCommandFence Fence;
 		ENQUEUE_RENDER_COMMAND(GetPixelColors)
 		(
diff --git a/Engine/Plugins/VirtualProduction/DMX/DMXFixtures/Source/DMXFixtures/Public/DMXFixtureActor.h b/Engine/Plugins/VirtualProduction/DMX/DMXFixtures/Source/DMXFixtures/Public/DMXFixtureActor.h
--- a/Engine/Plugins/VirtualProduction/DMX/DMXFixtures/Source/DMXFixtures/Public/DMXFixtureActor.h
+++ b/Engine/Plugins/VirtualProduction/DMX/DMXFixtures/Source/DMXFixtures/Public/DMXFixtureActor.h
@@ -50,9 +50,6 @@ public:
 	//~ Begin DMXMVRFixtureActorInterface interface
 	virtual void OnMVRGetSupportedDMXAttributes_Implementation(TArray<FName>& OutAttributeNames, TArray<FName>& OutMatrixAttributeNames) const override;
 	//~ End DMXMVRFixtureActorInterface interface
-	
-	// UObject interface
-	virtual void PostLoad() override;
 
 	bool HasBeenInitialized;
 	float LensRadius;
@@ -196,9 +193,9 @@ public:
 public:
 	UE_DEPRECATED(5.1, "Use BeamQuality instead")
 	UPROPERTY(meta = (DeprecatedProperty, DeprecationMessage = "Use BeamQuality instead."))
-	float MinQuality;
+	float MinQuality_DEPRECATED;
 
 	UE_DEPRECATED(5.1, "Use ZoomQuality instead")
 	UPROPERTY(meta = (DeprecatedProperty, DeprecationMessage = "Use ZoomQuality instead."))
-	float MaxQuality;
+	float MaxQuality_DEPRECATED;
 };
diff --git a/Engine/Plugins/VirtualProduction/DMX/DMXFixtures/Source/DMXFixtures/Public/DMXFixtureComponent.h b/Engine/Plugins/VirtualProduction/DMX/DMXFixtures/Source/DMXFixtures/Public/DMXFixtureComponent.h
--- a/Engine/Plugins/VirtualProduction/DMX/DMXFixtures/Source/DMXFixtures/Public/DMXFixtureComponent.h
+++ b/Engine/Plugins/VirtualProduction/DMX/DMXFixtures/Source/DMXFixtures/Public/DMXFixtureComponent.h
@@ -20,7 +20,7 @@ struct FDMXChannelData
 {
 	GENERATED_BODY()
 
-	UPROPERTY(EditAnywhere, Category = "DMX Channel")
+	UPROPERTY(EditAnywhere, BlueprintReadOnly, Category = "DMX Channel")
 	FDMXAttributeName Name;
 
 	UPROPERTY(EditAnywhere, BlueprintReadOnly, Category = "DMX Channel")
@@ -80,9 +80,12 @@ public:
 	UFUNCTION(BlueprintCallable, Category = "DMX")
 	class ADMXFixtureActor* GetParentFixtureActor();
 
-	/** Reads pixel color in the middle of each "Texture" and output linear colors */
+	/** 
+	 * Reads pixel color in the middle of each "Texture" and output linear colors. 
+	 * Note, this function locks the GPU to read the color and is slow to call, even if Update TextureResource is not enabled.
+	 */
 	UFUNCTION(BlueprintCallable, Category = "DMX")
-	TArray<FLinearColor> GetTextureCenterColors(UTexture2D* TextureAtlas, int numTextures);
+	TArray<FLinearColor> GetTextureCenterColors(UTexture2D* TextureAtlas, int numTextures, bool bUpdateTextureResource = true);
 
 	/** Called each tick when interpolation is enabled, to calculate the next value */
 	UFUNCTION(BlueprintCallable, BlueprintImplementableEvent, Category = "DMX")
