DMX - Fix an issue in how zoom component sets the intensity of the spotlight of fixture actors, update Fixture Actor to compute the spotlight intensity the same way, even when no zoom component is present

diff --git a/Engine/Plugins/VirtualProduction/DMX/DMXFixtures/Source/DMXFixtures/Private/DMXFixtureActor.cpp b/Engine/Plugins/VirtualProduction/DMX/DMXFixtures/Source/DMXFixtures/Private/DMXFixtureActor.cpp
--- a/Engine/Plugins/VirtualProduction/DMX/DMXFixtures/Source/DMXFixtures/Private/DMXFixtureActor.cpp
+++ b/Engine/Plugins/VirtualProduction/DMX/DMXFixtures/Source/DMXFixtures/Private/DMXFixtureActor.cpp
@@ -125,6 +125,13 @@ void ADMXFixtureActor::InitializeFixture(UStaticMeshComponent* StaticMeshLens, U
 	HasBeenInitialized = true;
 }
 
+void ADMXFixtureActor::UpdateSpotLightIntensity()
+{
+	const float SterdianRatio = UE_PI * 2.f * (1.f - SpotLight->GetCosHalfConeAngle());
+
+	SpotLight->SetIntensity(LightIntensityMax * SpotlightIntensityScale * (1.f / SterdianRatio));
+}
+
 void ADMXFixtureActor::FeedFixtureData()
 {
 	// BeamQuality and ZoomQuality modulate the "stepSize" for the raymarch beam shader
@@ -170,7 +177,7 @@ void ADMXFixtureActor::FeedFixtureData()
 	}
 
 	// Set lights
-	SpotLight->SetIntensity(LightIntensityMax * SpotlightIntensityScale);
+	UpdateSpotLightIntensity();
 	SpotLight->SetTemperature(LightColorTemp);
 	SpotLight->SetCastShadows(LightCastShadow);
 	SpotLight->SetAttenuationRadius(LightDistanceMax);
@@ -234,7 +241,7 @@ void ADMXFixtureActor::SetSpotlightIntensityScale(float NewSpotlightIntensitySca
 		DynamicMaterialLens->SetScalarParameterValue("DMX Max Light Intensity", LightIntensityMax * SpotlightIntensityScale);
 	}
 
-	SpotLight->SetIntensity(LightIntensityMax * SpotlightIntensityScale);
+	UpdateSpotLightIntensity();
 }
 
 void ADMXFixtureActor::SetPointlightIntensityScale(float NewPointlightIntensityScale)
diff --git a/Engine/Plugins/VirtualProduction/DMX/DMXFixtures/Source/DMXFixtures/Public/DMXFixtureActor.h b/Engine/Plugins/VirtualProduction/DMX/DMXFixtures/Source/DMXFixtures/Public/DMXFixtureActor.h
--- a/Engine/Plugins/VirtualProduction/DMX/DMXFixtures/Source/DMXFixtures/Public/DMXFixtureActor.h
+++ b/Engine/Plugins/VirtualProduction/DMX/DMXFixtures/Source/DMXFixtures/Public/DMXFixtureActor.h
@@ -88,6 +88,14 @@ public:
 	virtual void PushNormalizedValuesPerAttribute(const FDMXNormalizedAttributeValueMap& ValuePerAttributeMap) override;
 	
 public:
+	/** 
+	 * Updates the spotlight intensity. 
+	 * Considers light intensity max, the spotlight intensity scale and the cone angle to compute the intensity. 
+	 * Should be used instead of setting the intensity of the spotlight directly.
+	 */
+	UFUNCTION(BlueprintCallable, Category = "DMX Fixture")
+	void UpdateSpotLightIntensity();
+
 	/** Sets the a new max light intensity */
 	UFUNCTION(BlueprintCallable, Category = "DMX Fixture")
 	void SetLightIntensityMax(float NewLightIntensityMax);
@@ -112,7 +120,6 @@ public:
 	UFUNCTION(BlueprintCallable, Category = "DMX Fixture")
 	void SetLightCastShadow(bool bLightShouldCastShadow);
 
-
 	// PARAMETERS---------------------------------
 
 	// Light intensity at 1 steradian (32.77deg half angle)
