CIS fix.

diff --git a/Engine/Plugins/VirtualProduction/DMX/DMXFixtures/Source/DMXFixtures/Private/DMXFixtureComponent.cpp b/Engine/Plugins/VirtualProduction/DMX/DMXFixtures/Source/DMXFixtures/Private/DMXFixtureComponent.cpp
--- a/Engine/Plugins/VirtualProduction/DMX/DMXFixtures/Source/DMXFixtures/Private/DMXFixtureComponent.cpp
+++ b/Engine/Plugins/VirtualProduction/DMX/DMXFixtures/Source/DMXFixtures/Private/DMXFixtureComponent.cpp
@@ -87,7 +87,7 @@ ADMXFixtureActor* UDMXFixtureComponent::GetParentFixtureActor()
 TArray<FLinearColor> UDMXFixtureComponent::GetTextureCenterColors(UTexture2D* Texture, int NumCells, bool bUpdateTextureResource)
 {
 	TArray<FLinearColor> PixelColorArray;
-	PixelColorArray.SetNumZeroed(NumCells, EAllowShrinking::True);
+	PixelColorArray.SetNumZeroed(NumCells, EAllowShrinking::Yes);
 	if (Texture)
 	{
 		if (bUpdateTextureResource)
