DMX - DMXFixtureComponentDouble now applies its value range correctly

diff --git a/Engine/Plugins/VirtualProduction/DMX/DMXFixtures/Source/DMXFixtures/Private/DMXFixtureComponentDouble.cpp b/Engine/Plugins/VirtualProduction/DMX/DMXFixtures/Source/DMXFixtures/Private/DMXFixtureComponentDouble.cpp
--- a/Engine/Plugins/VirtualProduction/DMX/DMXFixtures/Source/DMXFixtures/Private/DMXFixtureComponentDouble.cpp
+++ b/Engine/Plugins/VirtualProduction/DMX/DMXFixtures/Source/DMXFixtures/Private/DMXFixtureComponentDouble.cpp
@@ -17,19 +17,19 @@ void UDMXFixtureComponentDouble::PushNormalizedValuesPerAttribute(const FDMXNorm
 	const float* FirstTargetValuePtr = ValuePerAttribute.Map.Find(DMXChannel1.Name);
 	if (FirstTargetValuePtr)
 	{
-		float Channel1RemappedValue = NormalizedToAbsoluteValue(0, *FirstTargetValuePtr);
+		const float Channel1RemappedValue = NormalizedToAbsoluteValue(0, *FirstTargetValuePtr);
 
 		constexpr int32 FirstChannelIndex = 0;
-		SetTargetValue(FirstChannelIndex, *FirstTargetValuePtr);
+		SetTargetValue(FirstChannelIndex, Channel1RemappedValue);
 	}
 
 	const float* SecondTargetValuePtr = ValuePerAttribute.Map.Find(DMXChannel2.Name);
 	if (SecondTargetValuePtr)
 	{
-		float Channel2RemappedValue = NormalizedToAbsoluteValue(1, *SecondTargetValuePtr);
+		const float Channel2RemappedValue = NormalizedToAbsoluteValue(1, *SecondTargetValuePtr);
 
 		constexpr int32 SecondChannelIndex = 1;
-		SetTargetValue(SecondChannelIndex, *SecondTargetValuePtr);
+		SetTargetValue(SecondChannelIndex, Channel2RemappedValue);
 	}
 }
 
diff --git a/Engine/Plugins/VirtualProduction/DMX/DMXFixtures/Source/DMXFixtures/Public/DMXFixtureComponentDouble.h b/Engine/Plugins/VirtualProduction/DMX/DMXFixtures/Source/DMXFixtures/Public/DMXFixtureComponentDouble.h
--- a/Engine/Plugins/VirtualProduction/DMX/DMXFixtures/Source/DMXFixtures/Public/DMXFixtureComponentDouble.h
+++ b/Engine/Plugins/VirtualProduction/DMX/DMXFixtures/Source/DMXFixtures/Public/DMXFixtureComponentDouble.h
@@ -56,7 +56,7 @@ public:
 	/** Sets the target value for specified channel index. Interpolates to the value if bUseInterpolation is true. */
 	void SetTargetValue(int32 ChannelIndex, float Value);
 
-	//  Sets first value of the second channel. When interpolation is enabled this function should be called until the value is reached, else just once */
+	/** Sets first value of the second channel. When interpolation is enabled this function should be called until the value is reached, else just once */
 	UFUNCTION(BlueprintCallable, BlueprintImplementableEvent, Category = "DMX")
 	void SetChannel1ValueNoInterp(float Channel1Value);
 
